[Player Input Debugger] First pass at a new "Player Input Debugger" for UE.

This tool can basically display a bunch of useful debug data that you would normally have to use individual "showdebug" commands for for player input, CommonUI/Input, and physical devices.

#rnx
#Synced-CL 51470122

#Synced-CL 51470122

diff --git a/Engine/Plugins/PlayerInputDebugger/PlayerInputDebugger.uplugin b/Engine/Plugins/PlayerInputDebugger/PlayerInputDebugger.uplugin
new file mode 100644
--- /dev/null
+++ b/Engine/Plugins/PlayerInputDebugger/PlayerInputDebugger.uplugin
@@ -0,0 +1,35 @@
+{
+	"FileVersion": 3,
+	"Version": 1,
+	"VersionName": "0.1",
+	"FriendlyName": "PlayerInputDebugger",
+	"Description": "Plugin for debugging input related things",
+	"Category": "Other",
+	"CreatedBy": "Epic Games, Inc.",
+	"CreatedByURL": "https://epicgames.com",
+	"DocsURL": "",
+	"MarketplaceURL": "",
+	"CanContainContent": false,
+	"IsBetaVersion": true,
+	"IsExperimentalVersion": true,
+	"Installed": false,
+	"EnabledByDefault": false,
+	"EditorOnly": true,
+	"Modules": [
+		{
+			"Name": "PlayerInputDebugger",
+			"Type": "Editor",
+			"LoadingPhase": "PostEngineInit"
+		}
+	],
+	"Plugins": [
+		{
+			"Name": "EnhancedInput",
+			"Enabled": true
+		},
+		{
+			"Name": "CommonUI",
+			"Enabled": true
+		}
+	]
+}
diff --git a/Engine/Plugins/PlayerInputDebugger/Source/PlayerInputDebugger/PlayerInputDebugger.Build.cs b/Engine/Plugins/PlayerInputDebugger/Source/PlayerInputDebugger/PlayerInputDebugger.Build.cs
new file mode 100644
--- /dev/null
+++ b/Engine/Plugins/PlayerInputDebugger/Source/PlayerInputDebugger/PlayerInputDebugger.Build.cs
@@ -0,0 +1,34 @@
+// Copyright Epic Games, Inc. All Rights Reserved.
+
+using UnrealBuildTool;
+
+public class PlayerInputDebugger : ModuleRules
+{
+	public PlayerInputDebugger(ReadOnlyTargetRules Target) : base(Target)
+	{
+		PublicDependencyModuleNames.AddRange(
+			new string[]
+			{
+
+			});
+
+		PrivateDependencyModuleNames.AddRange(
+			new string[]
+			{
+				"Core",
+				"CoreUObject",
+				"Engine",
+				"EnhancedInput",
+				"Slate",
+				"SlateCore",
+				"ToolMenus",
+				"UnrealEd",
+				"InputCore",
+				"ApplicationCore",
+				"PropertyEditor",
+				"CommonUI",
+				"CommonInput",
+				"UMG",
+			});
+	}
+}
diff --git a/Engine/Plugins/PlayerInputDebugger/Source/PlayerInputDebugger/Private/PlayerInputDebuggerModule.cpp b/Engine/Plugins/PlayerInputDebugger/Source/PlayerInputDebugger/Private/PlayerInputDebuggerModule.cpp
new file mode 100644
--- /dev/null
+++ b/Engine/Plugins/PlayerInputDebugger/Source/PlayerInputDebugger/Private/PlayerInputDebuggerModule.cpp
@@ -0,0 +1,72 @@
+// Copyright Epic Games, Inc. All Rights Reserved.
+
+#include "PlayerInputDebuggerModule.h"
+
+#include "SPlayerInputDebugger.h"
+
+#include "Framework/Application/SlateApplication.h"
+#include "Framework/Docking/TabManager.h"
+#include "Modules/ModuleManager.h"
+#include "ToolMenus.h"
+#include "Widgets/Docking/SDockTab.h"
+
+#define LOCTEXT_NAMESPACE "FPlayerInputDebuggerModule"
+
+static const FName PlayerInputDebuggerTabName("PlayerInputDebugger");
+
+void FPlayerInputDebuggerModule::StartupModule()
+{
+	MenusStartupHandle = UToolMenus::RegisterStartupCallback(
+		FSimpleMulticastDelegate::FDelegate::CreateRaw(this, &FPlayerInputDebuggerModule::RegisterMenus));
+}
+
+void FPlayerInputDebuggerModule::ShutdownModule()
+{
+	UToolMenus::UnRegisterStartupCallback(MenusStartupHandle);
+	UToolMenus::UnRegisterStartupCallback(this);
+
+	if (UToolMenus* ToolMenus = UToolMenus::TryGet())
+	{
+		ToolMenus->UnregisterOwnerByName("PlayerInputDebugger");
+	}
+
+	FGlobalTabmanager::Get()->UnregisterNomadTabSpawner(PlayerInputDebuggerTabName);
+}
+
+void FPlayerInputDebuggerModule::RegisterMenus()
+{
+	FToolMenuOwnerScoped OwnerScoped("PlayerInputDebugger");
+
+	FGlobalTabmanager::Get()->RegisterNomadTabSpawner(
+		PlayerInputDebuggerTabName,
+		FOnSpawnTab::CreateRaw(this, &FPlayerInputDebuggerModule::SpawnTab))
+		.SetDisplayName(LOCTEXT("TabTitle", "Player Input Debugger"))
+		.SetMenuType(ETabSpawnerMenuType::Hidden);
+
+	if (UToolMenu* ToolsMenu = UToolMenus::Get()->ExtendMenu("MainFrame.MainMenu.Tools"))
+	{
+		FToolMenuSection& Section = ToolsMenu->FindOrAddSection("Debugging", LOCTEXT("DebuggingSection", "Debugging"));
+		Section.AddMenuEntry(
+			"PlayerInputDebugger",
+			LOCTEXT("MenuEntryLabel", "Player Input Debugger"),
+			LOCTEXT("MenuEntryTooltip", "Open the Player Input Debugger window to inspect input events, connected devices, and Enhanced Input state."),
+			FSlateIcon(),
+			FExecuteAction::CreateLambda([]()
+			{
+				FGlobalTabmanager::Get()->TryInvokeTab(FName("PlayerInputDebugger"));
+			}));
+	}
+}
+
+TSharedRef<SDockTab> FPlayerInputDebuggerModule::SpawnTab(const FSpawnTabArgs& Args)
+{
+	return SNew(SDockTab)
+		.TabRole(ETabRole::NomadTab)
+		[
+			SNew(SPlayerInputDebugger)
+		];
+}
+
+#undef LOCTEXT_NAMESPACE
+
+IMPLEMENT_MODULE(FPlayerInputDebuggerModule, PlayerInputDebugger)
diff --git a/Engine/Plugins/PlayerInputDebugger/Source/PlayerInputDebugger/Private/PlayerInputDebuggerModule.h b/Engine/Plugins/PlayerInputDebugger/Source/PlayerInputDebugger/Private/PlayerInputDebuggerModule.h
new file mode 100644
--- /dev/null
+++ b/Engine/Plugins/PlayerInputDebugger/Source/PlayerInputDebugger/Private/PlayerInputDebuggerModule.h
@@ -0,0 +1,22 @@
+// Copyright Epic Games, Inc. All Rights Reserved.
+
+#pragma once
+
+#include "CoreMinimal.h"
+#include "Modules/ModuleInterface.h"
+
+class FSpawnTabArgs;
+class SDockTab;
+
+class FPlayerInputDebuggerModule : public IModuleInterface
+{
+public:
+	virtual void StartupModule() override;
+	virtual void ShutdownModule() override;
+
+private:
+	void RegisterMenus();
+	TSharedRef<SDockTab> SpawnTab(const FSpawnTabArgs& Args);
+
+	FDelegateHandle MenusStartupHandle;
+};
diff --git a/Engine/Plugins/PlayerInputDebugger/Source/PlayerInputDebugger/Private/SEnhancedInputTab.cpp b/Engine/Plugins/PlayerInputDebugger/Source/PlayerInputDebugger/Private/SEnhancedInputTab.cpp
new file mode 100644
--- /dev/null
+++ b/Engine/Plugins/PlayerInputDebugger/Source/PlayerInputDebugger/Private/SEnhancedInputTab.cpp
@@ -0,0 +1,652 @@
+// Copyright Epic Games, Inc. All Rights Reserved.
+
+#include "SEnhancedInputTab.h"
+
+#include "EnhancedActionKeyMapping.h"
+#include "EnhancedInputSubsystems.h"
+#include "EnhancedPlayerInput.h"
+#include "InputAction.h"
+#include "InputMappingContext.h"
+#include "UObject/UObjectIterator.h"
+#include "Engine/LocalPlayer.h"
+#include "Framework/Application/SlateApplication.h"
+#include "Styling/AppStyle.h"
+#include "UserSettings/EnhancedInputUserSettings.h"
+#include "PropertyCustomizationHelpers.h"
+#include "Widgets/Images/SImage.h"
+#include "Widgets/Input/SButton.h"
+#include "Widgets/Input/SEditableTextBox.h"
+#include "Widgets/Layout/SBorder.h"
+#include "Widgets/Layout/SBox.h"
+#include "Widgets/Layout/SExpandableArea.h"
+#include "Widgets/Text/STextBlock.h"
+#include "Widgets/Views/SListView.h"
+#include "Widgets/Views/STableRow.h"
+
+#define LOCTEXT_NAMESPACE "SEnhancedInputTab"
+
+namespace EnhancedInputTab
+{
+	static const float PriorityColumnWidth = 65.f;
+	static const float ContextColumnWidth  = 200.f;
+	static const float ActionColumnWidth   = 200.f;
+	static const float KeyColumnWidth      = 160.f;
+	static const float StateColumnWidth    = 90.f;
+	static const float ValueColumnWidth    = 150.f;
+}
+
+SEnhancedInputTab::~SEnhancedInputTab()
+{
+}
+
+void SEnhancedInputTab::SetPlayerController(APlayerController* PC)
+{
+	WeakPC = PC;
+}
+
+static FSlateColor GetTriggerEventColor(uint8 Ordinal)
+{
+	// None=0, Triggered=1, Started=2, Ongoing=3, Canceled=4, Completed=5
+	switch (Ordinal)
+	{
+		case 1: return FLinearColor(0.2f, 0.9f, 0.2f);  // Triggered - green
+		case 2: return FLinearColor(0.2f, 0.8f, 0.9f);  // Started   - cyan
+		case 3: return FLinearColor(0.9f, 0.9f, 0.2f);  // Ongoing   - yellow
+		case 4: return FLinearColor(0.9f, 0.2f, 0.2f);  // Canceled  - red
+		case 5: return FLinearColor(0.5f, 0.5f, 0.5f);  // Completed - gray
+		default: return FLinearColor(0.8f, 0.8f, 0.8f); // None      - dim white
+	}
+}
+
+class SEnhancedInputContextHeaderRow : public SMultiColumnTableRow<TSharedPtr<FEnhancedInputDisplayItem>>
+{
+public:
+	SLATE_BEGIN_ARGS(SEnhancedInputContextHeaderRow) {}
+		SLATE_ARGUMENT(TSharedPtr<FEnhancedInputDisplayItem>, Item)
+	SLATE_END_ARGS()
+
+	void Construct(const FArguments& InArgs, const TSharedRef<STableViewBase>& InOwnerTableView)
+	{
+		Item = InArgs._Item;
+		SMultiColumnTableRow<TSharedPtr<FEnhancedInputDisplayItem>>::Construct(
+			FSuperRowType::FArguments().Padding(FMargin(2.f, 2.f)),
+			InOwnerTableView
+		);
+	}
+
+	virtual TSharedRef<SWidget> GenerateWidgetForColumn(const FName& ColumnName) override
+	{
+		check(Item.IsValid());
+
+		if (ColumnName == TEXT("Priority"))
+		{
+			return SNew(STextBlock)
+				.Text(FText::AsNumber(Item->Priority))
+				.Font(FAppStyle::GetFontStyle("BoldFont"))
+				.ColorAndOpacity(FLinearColor(0.9f, 0.8f, 0.2f));
+		}
+
+		if (ColumnName == TEXT("Context"))
+		{
+			return SNew(STextBlock)
+				.Text(Item->ContextName)
+				.Font(FAppStyle::GetFontStyle("BoldFont"))
+				.ColorAndOpacity(FLinearColor(0.9f, 0.8f, 0.2f));
+		}
+
+		return SNew(STextBlock).Text(FText::GetEmpty());
+	}
+
+private:
+	TSharedPtr<FEnhancedInputDisplayItem> Item;
+};
+
+class SEnhancedInputRow : public SMultiColumnTableRow<TSharedPtr<FEnhancedInputDisplayItem>>
+{
+public:
+	SLATE_BEGIN_ARGS(SEnhancedInputRow) {}
+		SLATE_ARGUMENT(TSharedPtr<FEnhancedInputDisplayItem>, Item)
+	SLATE_END_ARGS()
+
+	void Construct(const FArguments& InArgs, const TSharedRef<STableViewBase>& InOwnerTableView)
+	{
+		Item = InArgs._Item;
+		SMultiColumnTableRow<TSharedPtr<FEnhancedInputDisplayItem>>::Construct(
+			FSuperRowType::FArguments().Padding(FMargin(2.f, 1.f)),
+			InOwnerTableView
+		);
+	}
+
+	virtual TSharedRef<SWidget> GenerateWidgetForColumn(const FName& ColumnName) override
+	{
+		check(Item.IsValid());
+
+		static const TArray<FText> TriggerEventNames =
+		{
+			LOCTEXT("TENone",      "None"),
+			LOCTEXT("TETriggered", "Triggered"),
+			LOCTEXT("TEStarted",   "Started"),
+			LOCTEXT("TEOngoing",   "Ongoing"),
+			LOCTEXT("TECanceled",  "Canceled"),
+			LOCTEXT("TECompleted", "Completed"),
+		};
+
+		const FText StateName = TriggerEventNames.IsValidIndex(Item->TriggerEventOrdinal)
+			? TriggerEventNames[Item->TriggerEventOrdinal]
+			: FText::AsNumber(Item->TriggerEventOrdinal);
+
+		// Priority and Context columns are blank for action rows
+		if (ColumnName == TEXT("Priority") || ColumnName == TEXT("Context"))
+		{
+			return SNew(STextBlock).Text(FText::GetEmpty());
+		}
+
+		if (ColumnName == TEXT("Action"))
+		{
+			return SNew(STextBlock)
+				.Text(Item->ActionName)
+				.Font(FAppStyle::GetFontStyle("NormalFont"));
+		}
+
+		if (ColumnName == TEXT("Key"))
+		{
+			return SNew(STextBlock)
+				.Text(FText::FromString(Item->BoundKey.ToString()))
+				.Font(FAppStyle::GetFontStyle("NormalFont"))
+				.ColorAndOpacity(FLinearColor(0.7f, 0.7f, 0.7f));
+		}
+
+		if (ColumnName == TEXT("State"))
+		{
+			return SNew(STextBlock)
+				.Text(StateName)
+				.Font(FAppStyle::GetFontStyle("NormalFont"))
+				.ColorAndOpacity(GetTriggerEventColor(Item->TriggerEventOrdinal));
+		}
+
+		if (ColumnName == TEXT("Value"))
+		{
+			return SNew(STextBlock)
+				.Text(Item->ValueText)
+				.Font(FAppStyle::GetFontStyle("NormalFont"));
+		}
+
+		return SNew(STextBlock).Text(FText::GetEmpty());
+	}
+
+private:
+	TSharedPtr<FEnhancedInputDisplayItem> Item;
+};
+
+void SEnhancedInputTab::Construct(const FArguments& InArgs)
+{
+	SetPlayerController(InArgs._PlayerController.Get());
+
+	SetCanTick(true);
+
+	TSharedRef<SHeaderRow> HeaderRow = SNew(SHeaderRow)
+		+ SHeaderRow::Column("Priority")
+		.DefaultLabel(LOCTEXT("ColPriority", "Priority"))
+		.ManualWidth(EnhancedInputTab::PriorityColumnWidth)
+
+		+ SHeaderRow::Column("Context")
+		.DefaultLabel(LOCTEXT("ColContext", "Context"))
+		.ManualWidth(EnhancedInputTab::ContextColumnWidth)
+
+		+ SHeaderRow::Column("Action")
+		.DefaultLabel(LOCTEXT("ColAction", "Action"))
+		.ManualWidth(EnhancedInputTab::ActionColumnWidth)
+
+		+ SHeaderRow::Column("Key")
+		.DefaultLabel(LOCTEXT("ColKey", "Key"))
+		.ManualWidth(EnhancedInputTab::KeyColumnWidth)
+
+		+ SHeaderRow::Column("State")
+		.DefaultLabel(LOCTEXT("ColState", "Trigger State"))
+		.ManualWidth(EnhancedInputTab::StateColumnWidth)
+
+		+ SHeaderRow::Column("Value")
+		.DefaultLabel(LOCTEXT("ColValue", "Value"))
+		.ManualWidth(EnhancedInputTab::ValueColumnWidth);
+
+	ChildSlot
+	[
+		SNew(SVerticalBox)
+
+		+ SVerticalBox::Slot()
+		.AutoHeight()
+		.Padding(8.f, 8.f, 8.f, 4.f)
+		[
+			SNew(STextBlock)
+			.Text(LOCTEXT("Header", "Active Mapping Contexts (highest priority first)"))
+			.Font(FAppStyle::GetFontStyle("HeadingExtraSmall"))
+		]
+
+		+ SVerticalBox::Slot()
+		.AutoHeight()
+		[
+			SNew(SExpandableArea)
+			.AreaTitle(LOCTEXT("ContextFiltersTitle", "Context Filters"))
+			.InitiallyCollapsed(false)
+			.BodyContent()
+			[
+				SAssignNew(FilterTagsContainer, SBox)
+				.Padding(FMargin(4.f, 4.f))
+			]
+		]
+
+		+ SVerticalBox::Slot()
+		.AutoHeight()
+		[
+			SNew(SExpandableArea)
+			.AreaTitle(LOCTEXT("KeyProfilesTitle", "Mappable Key Profiles"))
+			.InitiallyCollapsed(false)
+			.BodyContent()
+			[
+				SNew(SVerticalBox)
+
+				// Active profile display
+				+ SVerticalBox::Slot()
+				.AutoHeight()
+				.Padding(4.f, 4.f, 4.f, 2.f)
+				[
+					SNew(SHorizontalBox)
+
+					+ SHorizontalBox::Slot()
+					.AutoWidth()
+					.VAlign(VAlign_Center)
+					.Padding(0.f, 0.f, 6.f, 0.f)
+					[
+						SNew(STextBlock)
+						.Text(LOCTEXT("ActiveProfileLabel", "Active Profile:"))
+					]
+
+					+ SHorizontalBox::Slot()
+					.AutoWidth()
+					.VAlign(VAlign_Center)
+					[
+						SNew(STextBlock)
+						.Text_Lambda([this]() -> FText
+						{
+							APlayerController* PC = WeakPC.Get();
+							if (!PC) { return LOCTEXT("NoPC", "(no player controller)"); }
+							ULocalPlayer* LP = PC->GetLocalPlayer();
+							if (!LP) { return LOCTEXT("NoLP", "(no local player)"); }
+							UEnhancedInputLocalPlayerSubsystem* EIS = LP->GetSubsystem<UEnhancedInputLocalPlayerSubsystem>();
+							if (!EIS) { return LOCTEXT("NoEIS", "(no EIS)"); }
+							UEnhancedInputUserSettings* Settings = EIS->GetUserSettings();
+							if (!Settings) { return LOCTEXT("NoSettings", "(no user settings)"); }
+							return FText::FromString(Settings->GetActiveKeyProfileId());
+						})
+						.ColorAndOpacity(FLinearColor(0.2f, 0.9f, 0.2f))
+					]
+				]
+
+				// Profile ID text entry
+				+ SVerticalBox::Slot()
+				.AutoHeight()
+				.Padding(4.f, 2.f)
+				[
+					SNew(SHorizontalBox)
+
+					+ SHorizontalBox::Slot()
+					.AutoWidth()
+					.VAlign(VAlign_Center)
+					.Padding(0.f, 0.f, 6.f, 0.f)
+					[
+						SNew(STextBlock)
+						.Text(LOCTEXT("ProfileIdLabel", "Profile ID:"))
+					]
+
+					+ SHorizontalBox::Slot()
+					.FillWidth(1.f)
+					[
+						SAssignNew(ProfileIdInput, SEditableTextBox)
+						.HintText(LOCTEXT("ProfileIdHint", "Enter profile identifier..."))
+					]
+				]
+
+				// Action buttons
+				+ SVerticalBox::Slot()
+				.AutoHeight()
+				.Padding(4.f, 2.f, 4.f, 4.f)
+				[
+					SNew(SHorizontalBox)
+
+					+ SHorizontalBox::Slot()
+					.AutoWidth()
+					.Padding(0.f, 0.f, 6.f, 0.f)
+					[
+						SNew(SButton)
+						.Text(LOCTEXT("SetProfileBtn", "Set Mapping Profile"))
+						.OnClicked_Lambda([this]() -> FReply
+						{
+							if (!ProfileIdInput.IsValid()) { return FReply::Handled(); }
+							const FString ProfileId = ProfileIdInput->GetText().ToString();
+							if (ProfileId.IsEmpty()) { return FReply::Handled(); }
+							APlayerController* PC = WeakPC.Get();
+							if (!PC) { return FReply::Handled(); }
+							ULocalPlayer* LP = PC->GetLocalPlayer();
+							if (!LP) { return FReply::Handled(); }
+							UEnhancedInputLocalPlayerSubsystem* EIS = LP->GetSubsystem<UEnhancedInputLocalPlayerSubsystem>();
+							if (!EIS) { return FReply::Handled(); }
+							UEnhancedInputUserSettings* Settings = EIS->GetUserSettings();
+							if (!Settings) { return FReply::Handled(); }
+							Settings->SetActiveKeyProfile(ProfileId);
+							return FReply::Handled();
+						})
+					]
+
+					+ SHorizontalBox::Slot()
+					.AutoWidth()
+					[
+						SNew(SButton)
+						.Text(LOCTEXT("ResetProfileBtn", "Reset to Default"))
+						.OnClicked_Lambda([this]() -> FReply
+						{
+							APlayerController* PC = WeakPC.Get();
+							if (!PC) { return FReply::Handled(); }
+							ULocalPlayer* LP = PC->GetLocalPlayer();
+							if (!LP) { return FReply::Handled(); }
+							UEnhancedInputLocalPlayerSubsystem* EIS = LP->GetSubsystem<UEnhancedInputLocalPlayerSubsystem>();
+							if (!EIS) { return FReply::Handled(); }
+							UEnhancedInputUserSettings* Settings = EIS->GetUserSettings();
+							if (!Settings) { return FReply::Handled(); }
+							Settings->SetKeyProfileToDefault();
+							return FReply::Handled();
+						})
+					]
+				]
+			]
+		]
+
+		+ SVerticalBox::Slot()
+		.FillHeight(1.f)
+		[
+			SAssignNew(ListView, SListView<TSharedPtr<FEnhancedInputDisplayItem>>)
+			.ListItemsSource(&DisplayItems)
+			.OnGenerateRow(this, &SEnhancedInputTab::OnGenerateRow)
+			.HeaderRow(HeaderRow)
+			.SelectionMode(ESelectionMode::None)
+		]
+	];
+
+	RebuildFilterTagsWidget();
+}
+
+void SEnhancedInputTab::Tick(const FGeometry& AllottedGeometry, const double InCurrentTime, const float InDeltaTime)
+{
+	SCompoundWidget::Tick(AllottedGeometry, InCurrentTime, InDeltaTime);
+	RebuildDisplayItems();
+}
+
+void SEnhancedInputTab::RemoveFilterContext(int32 Index)
+{
+	if (FilterContexts.IsValidIndex(Index))
+	{
+		FilterContexts.RemoveAt(Index);
+		RebuildFilterTagsWidget();
+	}
+}
+
+void SEnhancedInputTab::RebuildFilterTagsWidget()
+{
+	TSharedRef<SVerticalBox> Content = SNew(SVerticalBox);
+
+	// Count valid (non-null) filter entries for the hint text.
+	const int32 ValidCount = FilterContexts.FilterByPredicate(
+		[](const TWeakObjectPtr<UInputMappingContext>& W) { return W.IsValid(); }).Num();
+
+	// Header row: status hint + "+" add-slot button
+	Content->AddSlot()
+	.AutoHeight()
+	.Padding(0.f, 0.f, 0.f, 4.f)
+	[
+		SNew(SHorizontalBox)
+
+		+ SHorizontalBox::Slot()
+		.FillWidth(1.f)
+		.VAlign(VAlign_Center)
+		[
+			SNew(STextBlock)
+			.Text(ValidCount == 0
+				? LOCTEXT("FilterNone", "No filter active — showing all contexts")
+				: FText::Format(LOCTEXT("FilterActive", "Filtering to {0} context(s)"), FText::AsNumber(ValidCount)))
+			.ColorAndOpacity(ValidCount == 0
+				? FSlateColor(FLinearColor(0.5f, 0.5f, 0.5f))
+				: FSlateColor(FSlateColor::UseForeground()))
+		]
+
+		+ SHorizontalBox::Slot()
+		.AutoWidth()
+		.VAlign(VAlign_Center)
+		[
+			SNew(SButton)
+			.ButtonStyle(FAppStyle::Get(), "SimpleButton")
+			.ToolTipText(LOCTEXT("AddFilterTooltip", "Add a context filter slot"))
+			.OnClicked_Lambda([this]() -> FReply
+			{
+				FilterContexts.Add(nullptr);
+				RebuildFilterTagsWidget();
+				return FReply::Handled();
+			})
+			[
+				SNew(SImage)
+				.Image(FAppStyle::GetBrush("Icons.Plus"))
+				.DesiredSizeOverride(FVector2D(16.f, 16.f))
+			]
+		]
+	];
+
+	// One asset-picker row per filter slot (matches Details panel TArray<UInputMappingContext*> style).
+	for (int32 i = 0; i < FilterContexts.Num(); ++i)
+	{
+		const int32 CapturedIndex = i;
+
+		Content->AddSlot()
+		.AutoHeight()
+		.Padding(0.f, 1.f)
+		[
+			SNew(SHorizontalBox)
+
+			// Array index badge
+			+ SHorizontalBox::Slot()
+			.AutoWidth()
+			.VAlign(VAlign_Center)
+			.Padding(0.f, 0.f, 6.f, 0.f)
+			[
+				SNew(STextBlock)
+				.Text(FText::AsNumber(i))
+				.ColorAndOpacity(FLinearColor(0.5f, 0.5f, 0.5f))
+				.MinDesiredWidth(18.f)
+			]
+
+			// UE property-style asset picker
+			+ SHorizontalBox::Slot()
+			.FillWidth(1.f)
+			[
+				SNew(SObjectPropertyEntryBox)
+				.AllowedClass(UInputMappingContext::StaticClass())
+				.ObjectPath(TAttribute<FString>::CreateLambda([this, CapturedIndex]() -> FString
+				{
+					if (FilterContexts.IsValidIndex(CapturedIndex))
+					{
+						if (const UInputMappingContext* IMC = FilterContexts[CapturedIndex].Get())
+						{
+							return IMC->GetPathName();
+						}
+					}
+					return FString();
+				}))
+				.OnObjectChanged(FOnSetObject::CreateLambda([this, CapturedIndex](const FAssetData& AssetData)
+				{
+					if (FilterContexts.IsValidIndex(CapturedIndex))
+					{
+						FilterContexts[CapturedIndex] = Cast<UInputMappingContext>(AssetData.GetAsset());
+						RebuildFilterTagsWidget();
+					}
+				}))
+				.AllowClear(true)
+				.DisplayUseSelected(true)
+				.DisplayBrowse(true)
+			]
+
+			// Delete-slot button
+			+ SHorizontalBox::Slot()
+			.AutoWidth()
+			.VAlign(VAlign_Center)
+			.Padding(4.f, 0.f, 0.f, 0.f)
+			[
+				SNew(SButton)
+				.ButtonStyle(FAppStyle::Get(), "SimpleButton")
+				.ToolTipText(LOCTEXT("RemoveFilterTooltip", "Remove this filter slot"))
+				.OnClicked_Lambda([this, CapturedIndex]() -> FReply
+				{
+					RemoveFilterContext(CapturedIndex);
+					return FReply::Handled();
+				})
+				[
+					SNew(SImage)
+					.Image(FAppStyle::GetBrush("Icons.Delete"))
+					.DesiredSizeOverride(FVector2D(16.f, 16.f))
+				]
+			]
+		];
+	}
+
+	FilterTagsContainer->SetContent(Content);
+}
+
+void SEnhancedInputTab::RebuildDisplayItems()
+{
+	DisplayItems.Empty();
+
+	APlayerController* PC = WeakPC.Get();
+	if (!PC)
+	{
+		ListView->RequestListRefresh();
+		return;
+	}
+
+	ULocalPlayer* LP = PC->GetLocalPlayer();
+	if (!LP)
+	{
+		ListView->RequestListRefresh();
+		return;
+	}
+
+	UEnhancedInputLocalPlayerSubsystem* EIS = LP->GetSubsystem<UEnhancedInputLocalPlayerSubsystem>();
+	if (!EIS)
+	{
+		ListView->RequestListRefresh();
+		return;
+	}
+
+	UEnhancedPlayerInput* PlayerInput = EIS->GetPlayerInput();
+	if (!PlayerInput)
+	{
+		ListView->RequestListRefresh();
+		return;
+	}
+
+	// Collect and sort contexts by priority (highest first).
+	TArray<TPair<const UInputMappingContext*, int32>> SortedContexts;
+	for (TObjectIterator<UInputMappingContext> ContextIt; ContextIt; ++ContextIt)
+	{
+		const UInputMappingContext* IMC = *ContextIt;
+		int32 Priority = -1;
+		if (IMC && EIS->HasMappingContext(IMC, Priority))
+		{
+			SortedContexts.Emplace(IMC, Priority);
+		}
+	}
+	SortedContexts.Sort([](const TPair<const UInputMappingContext*, int32>& A,
+	                       const TPair<const UInputMappingContext*, int32>& B)
+	{
+		return A.Value > B.Value;
+	});
+
+	for (const TPair<const UInputMappingContext*, int32>& ContextPair : SortedContexts)
+	{
+		const UInputMappingContext* IMC = ContextPair.Key;
+		const int32 Priority = ContextPair.Value;
+
+		// Apply context filter: only active when at least one slot holds a valid IMC.
+		// Null/unset slots are ignored so partially-filled filter lists don't hide everything.
+		const bool bHasValidFilter = FilterContexts.ContainsByPredicate(
+			[](const TWeakObjectPtr<UInputMappingContext>& W) { return W.IsValid(); });
+		if (bHasValidFilter)
+		{
+			const bool bPassesFilter = FilterContexts.ContainsByPredicate(
+				[IMC](const TWeakObjectPtr<UInputMappingContext>& W) { return W.Get() == IMC; });
+			if (!bPassesFilter)
+			{
+				continue;
+			}
+		}
+
+		// Context header row
+		TSharedPtr<FEnhancedInputDisplayItem> Header = MakeShared<FEnhancedInputDisplayItem>();
+		Header->bIsHeader = true;
+		Header->ContextName = FText::FromString(IMC->GetName());
+		Header->Priority = Priority;
+		DisplayItems.Add(MoveTemp(Header));
+
+		// One row per key mapping — collect into a local array so we can sort before committing.
+		TArray<TSharedPtr<FEnhancedInputDisplayItem>> ContextRows;
+		for (const FEnhancedActionKeyMapping& Mapping : IMC->GetMappings())
+		{
+			if (!Mapping.Action || Mapping.bShouldBeIgnored)
+			{
+				continue;
+			}
+
+			const FInputActionInstance* Instance = PlayerInput->FindActionInstanceData(Mapping.Action);
+
+			TSharedPtr<FEnhancedInputDisplayItem> Row = MakeShared<FEnhancedInputDisplayItem>();
+			Row->bIsHeader = false;
+			Row->ActionName = FText::FromString(Mapping.Action->GetName());
+			Row->BoundKey = Mapping.Key;
+
+			if (Instance)
+			{
+				Row->TriggerEventOrdinal = static_cast<uint8>(Instance->GetTriggerEvent());
+				Row->ValueText = FText::FromString(Instance->GetValue().ToString());
+			}
+			else
+			{
+				Row->TriggerEventOrdinal = 0; // ETriggerEvent::None
+				Row->ValueText = LOCTEXT("NoValue", "-");
+			}
+
+			ContextRows.Add(MoveTemp(Row));
+		}
+
+		// Sort by action name so all mappings for the same action are adjacent.
+		ContextRows.Sort([](const TSharedPtr<FEnhancedInputDisplayItem>& A, const TSharedPtr<FEnhancedInputDisplayItem>& B)
+		{
+			return A->ActionName.CompareTo(B->ActionName) < 0;
+		});
+
+		for (TSharedPtr<FEnhancedInputDisplayItem>& Row : ContextRows)
+		{
+			DisplayItems.Add(MoveTemp(Row));
+		}
+	}
+
+	ListView->RequestListRefresh();
+}
+
+TSharedRef<ITableRow> SEnhancedInputTab::OnGenerateRow(TSharedPtr<FEnhancedInputDisplayItem> Item, const TSharedRef<STableViewBase>& OwnerTable)
+{
+	if (Item->bIsHeader)
+	{
+		return SNew(SEnhancedInputContextHeaderRow, OwnerTable)
+			.Item(Item);
+	}
+
+	return SNew(SEnhancedInputRow, OwnerTable)
+		.Item(Item);
+}
+
+#undef LOCTEXT_NAMESPACE
diff --git a/Engine/Plugins/PlayerInputDebugger/Source/PlayerInputDebugger/Private/SEnhancedInputTab.h b/Engine/Plugins/PlayerInputDebugger/Source/PlayerInputDebugger/Private/SEnhancedInputTab.h
new file mode 100644
--- /dev/null
+++ b/Engine/Plugins/PlayerInputDebugger/Source/PlayerInputDebugger/Private/SEnhancedInputTab.h
@@ -0,0 +1,62 @@
+// Copyright Epic Games, Inc. All Rights Reserved.
+
+#pragma once
+
+#include "GameFramework/PlayerController.h"
+#include "Widgets/SCompoundWidget.h"
+#include "Widgets/Views/ITableRow.h"
+#include "Widgets/Views/SListView.h"
+
+class UInputMappingContext;
+class SEditableTextBox;
+
+// A single row in the Enhanced Input display: either a context header or an action mapping row.
+struct FEnhancedInputDisplayItem
+{
+	bool bIsHeader = false;
+
+	// Header fields
+	FText ContextName;
+	int32 Priority = 0;
+
+	// Action mapping row fields
+	FText ActionName;
+	FKey BoundKey;
+	uint8 TriggerEventOrdinal = 0;  // ETriggerEvent stored as uint8 to avoid including InputAction.h in header
+	FText ValueText;
+};
+
+// Tab showing active Enhanced Input mapping contexts, their actions, bound keys,
+// trigger states and values for the selected player controller. Refreshes every frame.
+class SEnhancedInputTab : public SCompoundWidget
+{
+public:
+	SLATE_BEGIN_ARGS(SEnhancedInputTab) {}
+		SLATE_ARGUMENT(TWeakObjectPtr<APlayerController>, PlayerController)
+	SLATE_END_ARGS()
+
+	void Construct(const FArguments& InArgs);
+	virtual ~SEnhancedInputTab();
+	virtual void Tick(const FGeometry& AllottedGeometry, const double InCurrentTime, const float InDeltaTime) override;
+
+	void SetPlayerController(APlayerController* PC);
+
+private:
+	void RebuildDisplayItems();
+	TSharedRef<ITableRow> OnGenerateRow(TSharedPtr<FEnhancedInputDisplayItem> Item, const TSharedRef<STableViewBase>& OwnerTable);
+
+	// Context filter helpers
+	void RemoveFilterContext(int32 Index);
+	void RebuildFilterTagsWidget();
+
+	TWeakObjectPtr<APlayerController> WeakPC;
+	TArray<TSharedPtr<FEnhancedInputDisplayItem>> DisplayItems;
+	TSharedPtr<SListView<TSharedPtr<FEnhancedInputDisplayItem>>> ListView;
+
+	// Contexts to show; empty = show all.
+	TArray<TWeakObjectPtr<UInputMappingContext>> FilterContexts;
+	TSharedPtr<SBox> FilterTagsContainer;
+
+	// Mappable Key Profiles section
+	TSharedPtr<SEditableTextBox> ProfileIdInput;
+};
diff --git a/Engine/Plugins/PlayerInputDebugger/Source/PlayerInputDebugger/Private/SGlobalInputInfoSection.cpp b/Engine/Plugins/PlayerInputDebugger/Source/PlayerInputDebugger/Private/SGlobalInputInfoSection.cpp
new file mode 100644
--- /dev/null
+++ b/Engine/Plugins/PlayerInputDebugger/Source/PlayerInputDebugger/Private/SGlobalInputInfoSection.cpp
@@ -0,0 +1,434 @@
+// Copyright Epic Games, Inc. All Rights Reserved.
+
+#include "SGlobalInputInfoSection.h"
+
+#include "Components/InputComponent.h"
+#include "CommonInputSubsystem.h"
+#include "Engine/GameViewportClient.h"
+#include "Engine/LocalPlayer.h"
+#include "Framework/Application/SlateApplication.h"
+#include "GameFramework/PlayerInput.h"
+#include "Input/CommonUIActionRouterBase.h"
+#include "Styling/AppStyle.h"
+#include "Widgets/Layout/SBorder.h"
+#include "Widgets/Layout/SBox.h"
+#include "Widgets/Layout/SExpandableArea.h"
+#include "Widgets/Text/STextBlock.h"
+#include "Widgets/Views/SListView.h"
+#include "Widgets/Views/STableRow.h"
+
+#define LOCTEXT_NAMESPACE "SGlobalInputInfoSection"
+
+// ── Stack row widget ──────────────────────────────────────────────────────────
+
+namespace
+{
+
+class SGlobalInputStackRow : public SMultiColumnTableRow<TSharedPtr<FGlobalStackRow>>
+{
+public:
+	SLATE_BEGIN_ARGS(SGlobalInputStackRow) {}
+		SLATE_ARGUMENT(TSharedPtr<FGlobalStackRow>, Item)
+	SLATE_END_ARGS()
+
+	void Construct(const FArguments& InArgs, const TSharedRef<STableViewBase>& InOwnerTableView)
+	{
+		Item = InArgs._Item;
+		SMultiColumnTableRow<TSharedPtr<FGlobalStackRow>>::Construct(
+			FSuperRowType::FArguments().Padding(FMargin(2.f, 1.f)),
+			InOwnerTableView
+		);
+	}
+
+	virtual TSharedRef<SWidget> GenerateWidgetForColumn(const FName& ColumnName) override
+	{
+		check(Item.IsValid());
+
+		static const FLinearColor BlockingColor(1.0f, 0.6f, 0.1f);
+
+		if (ColumnName == TEXT("StackPriority"))
+		{
+			return SNew(STextBlock)
+				.Text(FText::AsNumber(Item->Priority))
+				.Font(FAppStyle::GetFontStyle("NormalFont"))
+				.ColorAndOpacity(Item->bIsBlockingComponent ? BlockingColor : FLinearColor::White);
+		}
+
+		if (ColumnName == TEXT("StackComponent"))
+		{
+			return SNew(STextBlock)
+				.Text(FText::FromString(Item->ComponentName))
+				.Font(FAppStyle::GetFontStyle("NormalFont"))
+				.ColorAndOpacity(Item->bIsBlockingComponent ? BlockingColor : FLinearColor::White);
+		}
+
+		if (ColumnName == TEXT("StackOwner"))
+		{
+			return SNew(STextBlock)
+				.Text(FText::FromString(Item->OwnerName))
+				.Font(FAppStyle::GetFontStyle("NormalFont"))
+				.ColorAndOpacity(Item->bIsBlockingComponent ? BlockingColor : FLinearColor(0.6f, 0.6f, 0.6f));
+		}
+
+		return SNew(STextBlock).Text(FText::GetEmpty());
+	}
+
+private:
+	TSharedPtr<FGlobalStackRow> Item;
+};
+
+} // anonymous namespace
+
+// ── Construct ─────────────────────────────────────────────────────────────────
+
+void SGlobalInputInfoSection::Construct(const FArguments& InArgs)
+{
+	SetCanTick(true);
+
+	// Helper: one label + live-text value row for the config section.
+	auto MakeConfigRow = [](const FText& Label, TAttribute<FText> Value, float TopPad = 1.f, float BotPad = 1.f) -> TSharedRef<SWidget>
+	{
+		return SNew(SHorizontalBox)
+			+ SHorizontalBox::Slot().AutoWidth().VAlign(VAlign_Center)
+			[
+				SNew(STextBlock)
+				.Text(Label)
+				.ColorAndOpacity(FLinearColor(0.6f, 0.6f, 0.6f))
+				.MinDesiredWidth(160.f)
+			]
+			+ SHorizontalBox::Slot().AutoWidth().VAlign(VAlign_Center)
+			[
+				SNew(STextBlock).Text(Value)
+			];
+	};
+
+	// ── Input Configuration body ──────────────────────────────────────────────
+	TSharedRef<SVerticalBox> ConfigBody = SNew(SVerticalBox)
+
+		// Player Controller sub-header
+		+ SVerticalBox::Slot().AutoHeight().Padding(4.f, 6.f, 4.f, 4.f)
+		[
+			SNew(STextBlock)
+			.Text(LOCTEXT("PCSubHeader", "Player Controller"))
+			.Font(FAppStyle::GetFontStyle("BoldFont"))
+			.ColorAndOpacity(FLinearColor(0.9f, 0.8f, 0.3f))
+		]
+
+		+ SVerticalBox::Slot().AutoHeight()
+		[
+			MakeConfigRow(LOCTEXT("InputModeLabel", "Input Mode:"),
+				TAttribute<FText>::CreateLambda([this]() -> FText
+				{
+					APlayerController* PC = WeakPC.Get();
+					if (!PC) { return LOCTEXT("ConfigNA", "—"); }
+					return FText::FromString(PC->GetCurrentInputModeDebugString());
+				}))
+		]
+
+		+ SVerticalBox::Slot().AutoHeight()
+		[
+			MakeConfigRow(LOCTEXT("MouseLockLabel", "Mouse Lock Mode:"),
+				TAttribute<FText>::CreateLambda([this]() -> FText
+				{
+					APlayerController* PC = WeakPC.Get();
+					if (!PC) { return LOCTEXT("ConfigNA", "—"); }
+					ULocalPlayer* LP = PC->GetLocalPlayer();
+					if (!LP) { return LOCTEXT("ConfigNA", "—"); }
+					UGameViewportClient* GVC = LP->ViewportClient;
+					if (!GVC) { return LOCTEXT("ConfigNA", "—"); }
+					return UEnum::GetDisplayValueAsText(GVC->GetMouseLockMode());
+				}))
+		]
+
+		+ SVerticalBox::Slot().AutoHeight()
+		[
+			MakeConfigRow(LOCTEXT("MouseCaptureLabel", "Mouse Capture Mode:"),
+				TAttribute<FText>::CreateLambda([this]() -> FText
+				{
+					APlayerController* PC = WeakPC.Get();
+					if (!PC) { return LOCTEXT("ConfigNA", "—"); }
+					ULocalPlayer* LP = PC->GetLocalPlayer();
+					if (!LP) { return LOCTEXT("ConfigNA", "—"); }
+					UGameViewportClient* GVC = LP->ViewportClient;
+					if (!GVC) { return LOCTEXT("ConfigNA", "—"); }
+					return UEnum::GetDisplayValueAsText(GVC->GetMouseCaptureMode());
+				}))
+		]
+
+		+ SVerticalBox::Slot().AutoHeight()
+		[
+			MakeConfigRow(LOCTEXT("MousePosLabel", "Mouse Position:"),
+				TAttribute<FText>::CreateLambda([this]() -> FText
+				{
+					APlayerController* PC = WeakPC.Get();
+					if (!PC) { return LOCTEXT("ConfigNA", "—"); }
+					ULocalPlayer* LP = PC->GetLocalPlayer();
+					if (!LP) { return LOCTEXT("ConfigNA", "—"); }
+					UGameViewportClient* GVC = LP->ViewportClient;
+					if (!GVC) { return LOCTEXT("ConfigNA", "—"); }
+					FVector2D MousePos(EForceInit::ForceInitToZero);
+					if (!GVC->GetMousePosition(MousePos))
+					{
+						return LOCTEXT("MouseOffScreen", "(off-screen)");
+					}
+					return FText::FromString(FString::Printf(TEXT("%.0f, %.0f"), MousePos.X, MousePos.Y));
+				}))
+		]
+
+		// UI sub-header
+		+ SVerticalBox::Slot().AutoHeight().Padding(4.f, 8.f, 4.f, 4.f)
+		[
+			SNew(STextBlock)
+			.Text(LOCTEXT("UISubHeader", "UI"))
+			.Font(FAppStyle::GetFontStyle("BoldFont"))
+			.ColorAndOpacity(FLinearColor(0.9f, 0.8f, 0.3f))
+		]
+
+		+ SVerticalBox::Slot().AutoHeight()
+		[
+			MakeConfigRow(LOCTEXT("CommonInputTypeLabel", "Input Type:"),
+				TAttribute<FText>::CreateLambda([this]() -> FText
+				{
+					APlayerController* PC = WeakPC.Get();
+					if (!PC) { return LOCTEXT("ConfigNA", "—"); }
+					ULocalPlayer* LP = PC->GetLocalPlayer();
+					if (!LP) { return LOCTEXT("ConfigNA", "—"); }
+					UCommonInputSubsystem* CIS = LP->GetSubsystem<UCommonInputSubsystem>();
+					if (!CIS) { return LOCTEXT("ConfigNA", "—"); }
+					return UEnum::GetDisplayValueAsText(CIS->GetCurrentInputType());
+				}))
+		]
+
+		+ SVerticalBox::Slot().AutoHeight()
+		[
+			MakeConfigRow(LOCTEXT("GamepadNameLabel", "Gamepad Name:"),
+				TAttribute<FText>::CreateLambda([this]() -> FText
+				{
+					APlayerController* PC = WeakPC.Get();
+					if (!PC) { return LOCTEXT("ConfigNA", "—"); }
+					ULocalPlayer* LP = PC->GetLocalPlayer();
+					if (!LP) { return LOCTEXT("ConfigNA", "—"); }
+					UCommonInputSubsystem* CIS = LP->GetSubsystem<UCommonInputSubsystem>();
+					if (!CIS) { return LOCTEXT("ConfigNA", "—"); }
+					const FName GamepadName = CIS->GetCurrentGamepadName();
+					return GamepadName.IsNone()
+						? LOCTEXT("NoGamepad", "(none)")
+						: FText::FromName(GamepadName);
+				}))
+		]
+
+		+ SVerticalBox::Slot().AutoHeight()
+		[
+			MakeConfigRow(LOCTEXT("UIInputModeLabel", "Common UI Input Mode:"),
+				TAttribute<FText>::CreateLambda([this]() -> FText
+				{
+					APlayerController* PC = WeakPC.Get();
+					if (!PC) { return LOCTEXT("ConfigNA", "—"); }
+					ULocalPlayer* LP = PC->GetLocalPlayer();
+					if (!LP) { return LOCTEXT("ConfigNA", "—"); }
+					UCommonUIActionRouterBase* Router = LP->GetSubsystem<UCommonUIActionRouterBase>();
+					if (!Router) { return LOCTEXT("ConfigNA", "—"); }
+					return UEnum::GetDisplayValueAsText(Router->GetActiveInputMode());
+				}))
+		]
+
+		+ SVerticalBox::Slot().AutoHeight()
+		[
+			MakeConfigRow(LOCTEXT("UIInputConfigLabel", "Active Input Config:"),
+				TAttribute<FText>::CreateLambda([this]() -> FText
+				{
+					return CachedInputConfigString.IsEmpty()
+						? LOCTEXT("ConfigNA", "—")
+						: FText::FromString(CachedInputConfigString);
+				}), 1.f, 6.f)
+		];
+
+	// ── Input Component Stack body ────────────────────────────────────────────
+	TSharedRef<SBox> StackBody = SNew(SBox)
+		.MaxDesiredHeight(300.f)
+		[
+			SAssignNew(StackListView, SListView<TSharedPtr<FGlobalStackRow>>)
+			.ListItemsSource(&StackRows)
+			.OnGenerateRow(this, &SGlobalInputInfoSection::OnGenerateStackRow)
+			.HeaderRow(
+				SNew(SHeaderRow)
+				+ SHeaderRow::Column("StackPriority")
+				.DefaultLabel(LOCTEXT("ColStackPriority", "Priority"))
+				.ManualWidth(70.f)
+
+				+ SHeaderRow::Column("StackComponent")
+				.DefaultLabel(LOCTEXT("ColStackComponent", "Input Component"))
+				.ManualWidth(260.f)
+
+				+ SHeaderRow::Column("StackOwner")
+				.DefaultLabel(LOCTEXT("ColStackOwner", "Owner"))
+				.ManualWidth(260.f)
+			)
+			.SelectionMode(ESelectionMode::None)
+		];
+
+	// ── Root layout ───────────────────────────────────────────────────────────
+	ChildSlot
+	[
+		SNew(SExpandableArea)
+		.AreaTitle(LOCTEXT("SectionTitle", "Input Overview"))
+		.InitiallyCollapsed(false)
+		.BodyContent()
+		[
+			SNew(SVerticalBox)
+
+			// Focused widget row
+			+ SVerticalBox::Slot()
+			.AutoHeight()
+			.Padding(4.f, 4.f, 4.f, 4.f)
+			[
+				SNew(SHorizontalBox)
+
+				+ SHorizontalBox::Slot()
+				.AutoWidth()
+				.VAlign(VAlign_Center)
+				.Padding(0.f, 0.f, 6.f, 0.f)
+				[
+					SNew(STextBlock)
+					.Text(LOCTEXT("FocusedWidgetLabel", "Focused widget:"))
+					.ColorAndOpacity(FLinearColor(0.5f, 0.5f, 0.5f))
+				]
+
+				+ SHorizontalBox::Slot()
+				.AutoWidth()
+				.VAlign(VAlign_Center)
+				[
+					SNew(STextBlock)
+					.Text_Lambda([]() -> FText
+					{
+						if (!FSlateApplication::IsInitialized())
+						{
+							return FText::GetEmpty();
+						}
+						const TSharedPtr<SWidget> Focused = FSlateApplication::Get().GetUserFocusedWidget(0);
+						if (!Focused)
+						{
+							return LOCTEXT("FocusNone", "(none)");
+						}
+						return FText::FromString(Focused->GetType().ToString());
+					})
+				]
+			]
+
+			// Input Configuration
+			+ SVerticalBox::Slot()
+			.AutoHeight()
+			[
+				SNew(SExpandableArea)
+				.AreaTitle(LOCTEXT("InputConfigTitle", "Input Configuration"))
+				.InitiallyCollapsed(true)
+				.BodyContent()
+				[
+					ConfigBody
+				]
+			]
+
+			// Input Component Stack
+			+ SVerticalBox::Slot()
+			.AutoHeight()
+			[
+				SNew(SExpandableArea)
+				.AreaTitle(LOCTEXT("InputStackTitle", "Input Component Stack"))
+				.InitiallyCollapsed(true)
+				.BodyContent()
+				[
+					StackBody
+				]
+			]
+		]
+	];
+}
+
+SGlobalInputInfoSection::~SGlobalInputInfoSection()
+{
+	if (UCommonUIActionRouterBase* Router = WeakRouter.Get())
+	{
+		Router->OnActiveInputConfigChanged().Remove(InputConfigChangedHandle);
+	}
+	InputConfigChangedHandle.Reset();
+}
+
+void SGlobalInputInfoSection::Tick(const FGeometry& AllottedGeometry, const double InCurrentTime, const float InDeltaTime)
+{
+	SCompoundWidget::Tick(AllottedGeometry, InCurrentTime, InDeltaTime);
+	RebuildStackRows();
+}
+
+void SGlobalInputInfoSection::SetPlayerController(APlayerController* PC)
+{
+	// Unsubscribe from the previous router.
+	if (UCommonUIActionRouterBase* OldRouter = WeakRouter.Get())
+	{
+		OldRouter->OnActiveInputConfigChanged().Remove(InputConfigChangedHandle);
+		InputConfigChangedHandle.Reset();
+	}
+	WeakRouter.Reset();
+	CachedInputConfigString.Empty();
+
+	WeakPC = PC;
+
+	// Subscribe to the new router.
+	if (PC)
+	{
+		if (ULocalPlayer* LP = PC->GetLocalPlayer())
+		{
+			if (UCommonUIActionRouterBase* Router = LP->GetSubsystem<UCommonUIActionRouterBase>())
+			{
+				WeakRouter = Router;
+				InputConfigChangedHandle = Router->OnActiveInputConfigChanged().AddLambda(
+					[this](const FUIInputConfig& NewConfig)
+					{
+						CachedInputConfigString = NewConfig.ToString();
+					});
+			}
+		}
+	}
+}
+
+void SGlobalInputInfoSection::RebuildStackRows()
+{
+	StackRows.Empty();
+
+	APlayerController* PC = WeakPC.Get();
+	if (PC)
+	{
+#if !(UE_BUILD_SHIPPING || UE_BUILD_TEST)
+		for (const TWeakObjectPtr<const UInputComponent>& WeakIC : PC->Debug_GetMostRecentInputStack())
+		{
+			if (const UInputComponent* IC = WeakIC.Get())
+			{
+				TSharedPtr<FGlobalStackRow> Row = MakeShared<FGlobalStackRow>();
+				Row->Priority = IC->Priority;
+				Row->ComponentName = IC->GetName();
+				Row->bIsBlockingComponent = IC->bBlockInput;
+
+				if (const AActor* Owner = IC->GetOwner())
+				{
+					Row->OwnerName = Owner->GetActorNameOrLabel();
+				}
+				else if (const UObject* Outer = IC->GetOuter())
+				{
+					Row->OwnerName = Outer->GetName();
+				}
+
+				StackRows.Add(MoveTemp(Row));
+			}
+		}
+#endif	// #if !(UE_BUILD_SHIPPING || UE_BUILD_TEST)
+	}
+
+	StackListView->RequestListRefresh();
+}
+
+TSharedRef<ITableRow> SGlobalInputInfoSection::OnGenerateStackRow(TSharedPtr<FGlobalStackRow> Item, const TSharedRef<STableViewBase>& OwnerTable)
+{
+	return SNew(SGlobalInputStackRow, OwnerTable)
+		.Item(Item);
+}
+
+#undef LOCTEXT_NAMESPACE
diff --git a/Engine/Plugins/PlayerInputDebugger/Source/PlayerInputDebugger/Private/SGlobalInputInfoSection.h b/Engine/Plugins/PlayerInputDebugger/Source/PlayerInputDebugger/Private/SGlobalInputInfoSection.h
new file mode 100644
--- /dev/null
+++ b/Engine/Plugins/PlayerInputDebugger/Source/PlayerInputDebugger/Private/SGlobalInputInfoSection.h
@@ -0,0 +1,52 @@
+// Copyright Epic Games, Inc. All Rights Reserved.
+
+#pragma once
+
+#include "GameFramework/PlayerController.h"
+#include "Widgets/SCompoundWidget.h"
+#include "Widgets/Views/ITableRow.h"
+#include "Widgets/Views/SListView.h"
+
+class UCommonUIActionRouterBase;
+
+// Stack row used by the global input overview.
+struct FGlobalStackRow
+{
+	int32 Priority = 0;
+	FString ComponentName;
+	FString OwnerName;
+	bool bIsBlockingComponent = false;
+};
+
+// Persistent collapsible section shown at the top of the debugger window
+// regardless of which tab is active. Displays:
+//   • Currently focused Slate widget
+//   • Input Configuration (input mode, mouse settings, CommonUI state)
+//   • Input Component Stack
+class SGlobalInputInfoSection : public SCompoundWidget
+{
+public:
+	SLATE_BEGIN_ARGS(SGlobalInputInfoSection) {}
+	SLATE_END_ARGS()
+
+	void Construct(const FArguments& InArgs);
+	virtual ~SGlobalInputInfoSection();
+	virtual void Tick(const FGeometry& AllottedGeometry, const double InCurrentTime, const float InDeltaTime) override;
+
+	void SetPlayerController(APlayerController* PC);
+
+private:
+	void RebuildStackRows();
+	TSharedRef<ITableRow> OnGenerateStackRow(TSharedPtr<FGlobalStackRow> Item, const TSharedRef<STableViewBase>& OwnerTable);
+
+	TWeakObjectPtr<APlayerController> WeakPC;
+
+	// Input Component Stack section
+	TArray<TSharedPtr<FGlobalStackRow>> StackRows;
+	TSharedPtr<SListView<TSharedPtr<FGlobalStackRow>>> StackListView;
+
+	// Active Input Config — cached from UCommonUIActionRouterBase::OnActiveInputConfigChanged()
+	FString CachedInputConfigString;
+	FDelegateHandle InputConfigChangedHandle;
+	TWeakObjectPtr<UCommonUIActionRouterBase> WeakRouter;
+};
diff --git a/Engine/Plugins/PlayerInputDebugger/Source/PlayerInputDebugger/Private/SInputDevicesTab.cpp b/Engine/Plugins/PlayerInputDebugger/Source/PlayerInputDebugger/Private/SInputDevicesTab.cpp
new file mode 100644
--- /dev/null
+++ b/Engine/Plugins/PlayerInputDebugger/Source/PlayerInputDebugger/Private/SInputDevicesTab.cpp
@@ -0,0 +1,326 @@
+// Copyright Epic Games, Inc. All Rights Reserved.
+
+#include "SInputDevicesTab.h"
+
+#include "GameFramework/InputDeviceSubsystem.h"
+#include "GameFramework/InputSettings.h"
+#include "GenericPlatform/GenericPlatformInputDeviceMapper.h"
+#include "Misc/CoreMiscDefines.h"
+#include "Styling/AppStyle.h"
+#include "Widgets/Layout/SBorder.h"
+#include "Widgets/Layout/SScrollBox.h"
+#include "Widgets/Text/STextBlock.h"
+
+#define LOCTEXT_NAMESPACE "SInputDevicesTab"
+
+static FText GetMappingPolicyText(EInputDeviceMappingPolicy Policy)
+{
+	switch (Policy)
+	{
+		case EInputDeviceMappingPolicy::UseManagedPlatformLogin:
+			return LOCTEXT("PolicyManaged",   "UseManagedPlatformLogin");
+		case EInputDeviceMappingPolicy::PrimaryUserSharesKeyboardAndFirstGamepad:
+			return LOCTEXT("PolicyPrimary",   "PrimaryUserSharesKeyboardAndFirstGamepad");
+		case EInputDeviceMappingPolicy::CreateUniquePlatformUserForEachDevice:
+			return LOCTEXT("PolicyUnique",    "CreateUniquePlatformUserForEachDevice");
+		case EInputDeviceMappingPolicy::MapAllDevicesToPrimaryUser:
+			return LOCTEXT("PolicyMapAll",    "MapAllDevicesToPrimaryUser");
+		default:
+			return LOCTEXT("PolicyInvalid",   "Invalid");
+	}
+}
+
+void SInputDevicesTab::Construct(const FArguments& InArgs)
+{
+	SetCanTick(false);
+
+	ChildSlot
+	[
+		SNew(SVerticalBox)
+
+		+ SVerticalBox::Slot()
+		.AutoHeight()
+		.Padding(8.f, 8.f, 8.f, 4.f)
+		[
+			SNew(STextBlock)
+			.Text(LOCTEXT("Header", "Connected Input Devices"))
+			.Font(FAppStyle::GetFontStyle("HeadingExtraSmall"))
+		]
+
+		+ SVerticalBox::Slot()
+		.AutoHeight()
+		.Padding(8.f, 0.f, 8.f, 8.f)
+		[
+			SNew(SHorizontalBox)
+
+			+ SHorizontalBox::Slot()
+			.AutoWidth()
+			.Padding(0.f, 0.f, 6.f, 0.f)
+			.VAlign(VAlign_Center)
+			[
+				SNew(STextBlock)
+				.Text(LOCTEXT("PolicyLabel", "Device Mapping Policy:"))
+				.ColorAndOpacity(FLinearColor(0.6f, 0.6f, 0.6f))
+			]
+
+			+ SHorizontalBox::Slot()
+			.AutoWidth()
+			.VAlign(VAlign_Center)
+			[
+				SNew(STextBlock)
+				.Text_Lambda([]() -> FText
+				{
+					if (const UInputPlatformSettings* Settings = UInputPlatformSettings::Get())
+					{
+						return GetMappingPolicyText(Settings->DeviceMappingPolicy);
+					}
+					return LOCTEXT("PolicyUnknown", "Unknown");
+				})
+			]
+		]
+
+		+ SVerticalBox::Slot()
+		.FillHeight(1.f)
+		[
+			SNew(SScrollBox)
+			+ SScrollBox::Slot()
+			[
+				SAssignNew(DeviceListContainer, SVerticalBox)
+			]
+		]
+	];
+
+	ConnectionChangedHandle = IPlatformInputDeviceMapper::Get().GetOnInputDeviceConnectionChange().AddRaw(
+		this, &SInputDevicesTab::OnDeviceConnectionChanged);
+
+	// Seed KnownDevices from currently active devices.
+	{
+		IPlatformInputDeviceMapper& Mapper = IPlatformInputDeviceMapper::Get();
+		UInputDeviceSubsystem* DeviceSubsystem = UInputDeviceSubsystem::Get();
+
+		TArray<FPlatformUserId> ActiveUsers;
+		Mapper.GetAllActiveUsers(ActiveUsers);
+
+		for (const FPlatformUserId& UserId : ActiveUsers)
+		{
+			TArray<FInputDeviceId> Devices;
+			Mapper.GetAllInputDevicesForUser(UserId, Devices);
+
+			for (const FInputDeviceId& DeviceId : Devices)
+			{
+				TSharedPtr<FKnownDevice> Entry = MakeShared<FKnownDevice>();
+				Entry->UserId = UserId;
+				Entry->LastState = Mapper.GetInputDeviceConnectionState(DeviceId);
+				if (DeviceSubsystem)
+				{
+					Entry->HardwareId = DeviceSubsystem->GetInputDeviceHardwareIdentifier(DeviceId);
+				}
+				KnownDevices.Add(DeviceId, MoveTemp(Entry));
+			}
+		}
+	}
+
+	RebuildDeviceList();
+}
+
+SInputDevicesTab::~SInputDevicesTab()
+{
+	IPlatformInputDeviceMapper::Get().GetOnInputDeviceConnectionChange().Remove(ConnectionChangedHandle);
+}
+
+void SInputDevicesTab::RebuildDeviceList()
+{
+	if (!DeviceListContainer.IsValid())
+	{
+		return;
+	}
+
+	DeviceListContainer->ClearChildren();
+
+	if (KnownDevices.IsEmpty())
+	{
+		DeviceListContainer->AddSlot()
+		.AutoHeight()
+		.Padding(8.f, 4.f)
+		[
+			SNew(STextBlock)
+			.Text(LOCTEXT("NoUsers", "No active platform users."))
+			.ColorAndOpacity(FLinearColor(0.5f, 0.5f, 0.5f))
+		];
+		return;
+	}
+
+	// Group known devices by platform user.
+	TMap<FPlatformUserId, TArray<FInputDeviceId>> UserDevices;
+	for (const TPair<FInputDeviceId, TSharedPtr<FKnownDevice>>& Pair : KnownDevices)
+	{
+		UserDevices.FindOrAdd(Pair.Value->UserId).Add(Pair.Key);
+	}
+
+	// Sort users: the PC's platform user first, then numerically by ID.
+	const FPlatformUserId PCUserId = WeakPC.IsValid() ? WeakPC->GetPlatformUserId() : FPlatformUserId();
+	const bool bHasPCUser = WeakPC.IsValid();
+
+	TArray<FPlatformUserId> SortedUsers;
+	UserDevices.GetKeys(SortedUsers);
+	SortedUsers.Sort([&PCUserId, bHasPCUser](const FPlatformUserId& A, const FPlatformUserId& B)
+	{
+		if (bHasPCUser)
+		{
+			const bool bAIsPC = (A == PCUserId);
+			const bool bBIsPC = (B == PCUserId);
+			if (bAIsPC != bBIsPC)
+			{
+				return bAIsPC;
+			}
+		}
+		return A.GetInternalId() < B.GetInternalId();
+	});
+
+	for (const FPlatformUserId& UserId : SortedUsers)
+	{
+		const TArray<FInputDeviceId>& Devices = UserDevices[UserId];
+
+		// Header color: green if any device is still connected, orange if all are disconnected.
+		const bool bAnyConnected = Devices.ContainsByPredicate([this](const FInputDeviceId& Id)
+		{
+			return KnownDevices[Id]->LastState == EInputDeviceConnectionState::Connected;
+		});
+
+		DeviceListContainer->AddSlot()
+		.AutoHeight()
+		.Padding(4.f, 6.f, 4.f, 2.f)
+		[
+			SNew(SBorder)
+			.BorderImage(FAppStyle::GetBrush("ToolPanel.GroupBorder"))
+			.Padding(FMargin(6.f, 3.f))
+			[
+				SNew(STextBlock)
+				.Text(FText::Format(LOCTEXT("UserHeader", "Platform User {0}  ({1} device(s))"),
+					FText::AsNumber(UserId.GetInternalId()),
+					FText::AsNumber(Devices.Num())))
+				.Font(FAppStyle::GetFontStyle("BoldFont"))
+				.ColorAndOpacity(bAnyConnected
+					? FLinearColor(0.2f, 0.9f, 0.2f)
+					: FLinearColor(0.9f, 0.5f, 0.1f))
+			]
+		];
+
+		for (const FInputDeviceId& DeviceId : Devices)
+		{
+			const TSharedPtr<FKnownDevice>& Device = KnownDevices[DeviceId];
+			const FText DeviceTypeText = GetDeviceTypeText(Device->HardwareId);
+			const FText StateText = GetConnectionStateText(Device->LastState);
+			const FSlateColor StateColor = GetConnectionStateColor(Device->LastState);
+
+			DeviceListContainer->AddSlot()
+			.AutoHeight()
+			.Padding(20.f, 2.f, 4.f, 2.f)
+			[
+				SNew(SHorizontalBox)
+
+				+ SHorizontalBox::Slot()
+				.AutoWidth()
+				.Padding(0.f, 0.f, 8.f, 0.f)
+				.VAlign(VAlign_Center)
+				[
+					SNew(STextBlock)
+					.Text(FText::Format(LOCTEXT("DeviceId", "Device {0}"), FText::AsNumber(DeviceId.GetId())))
+					.Font(FAppStyle::GetFontStyle("NormalFont"))
+				]
+
+				+ SHorizontalBox::Slot()
+				.AutoWidth()
+				.Padding(0.f, 0.f, 8.f, 0.f)
+				.VAlign(VAlign_Center)
+				[
+					SNew(STextBlock)
+					.Text(DeviceTypeText)
+					.Font(FAppStyle::GetFontStyle("NormalFont"))
+					.ColorAndOpacity(FLinearColor(0.7f, 0.7f, 0.7f))
+				]
+
+				+ SHorizontalBox::Slot()
+				.AutoWidth()
+				.Padding(0.f, 0.f, 8.f, 0.f)
+				.VAlign(VAlign_Center)
+				[
+					SNew(STextBlock)
+					.Text(Device->HardwareId.HardwareDeviceIdentifier.IsNone()
+						? LOCTEXT("UnknownHardware", "(unknown hardware)")
+						: FText::FromName(Device->HardwareId.HardwareDeviceIdentifier))
+					.Font(FAppStyle::GetFontStyle("NormalFont"))
+					.ColorAndOpacity(FLinearColor(0.5f, 0.5f, 0.5f))
+				]
+
+				+ SHorizontalBox::Slot()
+				.AutoWidth()
+				.VAlign(VAlign_Center)
+				[
+					SNew(STextBlock)
+					.Text(StateText)
+					.Font(FAppStyle::GetFontStyle("NormalFont"))
+					.ColorAndOpacity(StateColor)
+				]
+			];
+		}
+	}
+}
+
+void SInputDevicesTab::OnDeviceConnectionChanged(EInputDeviceConnectionState NewState, FPlatformUserId UserId, FInputDeviceId DeviceId)
+{
+	if (TSharedPtr<FKnownDevice>* Existing = KnownDevices.Find(DeviceId))
+	{
+		// Update the state of an already-known device.
+		(*Existing)->LastState = NewState;
+	}
+	else
+	{
+		// First time we've seen this device — add it to the cache.
+		TSharedPtr<FKnownDevice> Entry = MakeShared<FKnownDevice>();
+		Entry->UserId = UserId;
+		Entry->LastState = NewState;
+		if (UInputDeviceSubsystem* DeviceSubsystem = UInputDeviceSubsystem::Get())
+		{
+			Entry->HardwareId = DeviceSubsystem->GetInputDeviceHardwareIdentifier(DeviceId);
+		}
+		KnownDevices.Add(DeviceId, MoveTemp(Entry));
+	}
+
+	RebuildDeviceList();
+}
+
+FText SInputDevicesTab::GetConnectionStateText(EInputDeviceConnectionState State) const
+{
+	switch (State)
+	{
+		case EInputDeviceConnectionState::Connected:    return LOCTEXT("StateConnected",    "Connected");
+		case EInputDeviceConnectionState::Disconnected: return LOCTEXT("StateDisconnected", "Disconnected");
+		case EInputDeviceConnectionState::Unknown:      return LOCTEXT("StateUnknown",      "Unknown");
+		default:                                        return LOCTEXT("StateInvalid",      "Invalid");
+	}
+}
+
+FSlateColor SInputDevicesTab::GetConnectionStateColor(EInputDeviceConnectionState State) const
+{
+	switch (State)
+	{
+		case EInputDeviceConnectionState::Connected:    return FLinearColor(0.2f, 0.9f, 0.2f);
+		case EInputDeviceConnectionState::Disconnected: return FLinearColor(0.9f, 0.2f, 0.2f);
+		case EInputDeviceConnectionState::Unknown:      return FLinearColor(0.9f, 0.5f, 0.1f);
+		default:                                        return FLinearColor(0.9f, 0.1f, 0.1f);
+	}
+}
+
+FText SInputDevicesTab::GetDeviceTypeText(const FHardwareDeviceIdentifier& HardwareId) const
+{
+	switch (HardwareId.PrimaryDeviceType)
+	{
+		case EHardwareDevicePrimaryType::KeyboardAndMouse: return LOCTEXT("TypeKBM",     "[Keyboard & Mouse]");
+		case EHardwareDevicePrimaryType::Gamepad:          return LOCTEXT("TypeGamepad", "[Gamepad]");
+		case EHardwareDevicePrimaryType::Touch:            return LOCTEXT("TypeTouch",   "[Touch]");
+		default:                                           return LOCTEXT("TypeUnspec",  "[Unspecified]");
+	}
+}
+
+#undef LOCTEXT_NAMESPACE
diff --git a/Engine/Plugins/PlayerInputDebugger/Source/PlayerInputDebugger/Private/SInputDevicesTab.h b/Engine/Plugins/PlayerInputDebugger/Source/PlayerInputDebugger/Private/SInputDevicesTab.h
new file mode 100644
--- /dev/null
+++ b/Engine/Plugins/PlayerInputDebugger/Source/PlayerInputDebugger/Private/SInputDevicesTab.h
@@ -0,0 +1,46 @@
+// Copyright Epic Games, Inc. All Rights Reserved.
+
+#pragma once
+
+#include "CoreMinimal.h"
+#include "GenericPlatform/GenericPlatformInputDeviceMapper.h"
+#include "GameFramework/InputSettings.h"
+#include "GameFramework/PlayerController.h"
+#include "Widgets/SCompoundWidget.h"
+
+class SVerticalBox;
+
+// Cached record for a device that has been seen at least once.
+// Retained even after disconnection so it stays visible in the list.
+struct FKnownDevice
+{
+	FPlatformUserId UserId;
+	FHardwareDeviceIdentifier HardwareId;
+	EInputDeviceConnectionState LastState = EInputDeviceConnectionState::Unknown;
+};
+
+// Tab that shows input hardware devices per platform user.
+// Devices remain visible after disconnecting, shown with a "Disconnected" state.
+class SInputDevicesTab : public SCompoundWidget
+{
+public:
+	SLATE_BEGIN_ARGS(SInputDevicesTab) {}
+	SLATE_END_ARGS()
+
+	void Construct(const FArguments& InArgs);
+	virtual ~SInputDevicesTab();
+
+	void SetPlayerController(APlayerController* PC) { WeakPC = PC; RebuildDeviceList(); }
+
+private:
+	void RebuildDeviceList();
+	void OnDeviceConnectionChanged(EInputDeviceConnectionState NewState, FPlatformUserId UserId, FInputDeviceId DeviceId);
+	FText GetConnectionStateText(EInputDeviceConnectionState State) const;
+	FSlateColor GetConnectionStateColor(EInputDeviceConnectionState State) const;
+	FText GetDeviceTypeText(const FHardwareDeviceIdentifier& HardwareId) const;
+
+	TSharedPtr<SVerticalBox> DeviceListContainer;
+	FDelegateHandle ConnectionChangedHandle;
+	TMap<FInputDeviceId, TSharedPtr<FKnownDevice>> KnownDevices;
+	TWeakObjectPtr<APlayerController> WeakPC;
+};
diff --git a/Engine/Plugins/PlayerInputDebugger/Source/PlayerInputDebugger/Private/SPlayerInputDebugger.cpp b/Engine/Plugins/PlayerInputDebugger/Source/PlayerInputDebugger/Private/SPlayerInputDebugger.cpp
new file mode 100644
--- /dev/null
+++ b/Engine/Plugins/PlayerInputDebugger/Source/PlayerInputDebugger/Private/SPlayerInputDebugger.cpp
@@ -0,0 +1,330 @@
+// Copyright Epic Games, Inc. All Rights Reserved.
+
+#include "SPlayerInputDebugger.h"
+
+#include "SEnhancedInputTab.h"
+#include "SInputDevicesTab.h"
+#include "SPlayerInputEventsTab.h"
+#include "SGlobalInputInfoSection.h"
+
+#include "Editor.h"
+#include "Engine/Engine.h"
+#include "Engine/World.h"
+#include "GameFramework/GameModeBase.h"
+#include "GameFramework/PlayerController.h"
+#include "Styling/AppStyle.h"
+#include "Widgets/Images/SImage.h"
+#include "Widgets/Input/SButton.h"
+#include "Widgets/Input/SCheckBox.h"
+#include "Widgets/Input/SComboBox.h"
+#include "Widgets/Layout/SBorder.h"
+#include "Widgets/Layout/SWidgetSwitcher.h"
+#include "Widgets/Text/STextBlock.h"
+
+#define LOCTEXT_NAMESPACE "SPlayerInputDebugger"
+
+void SPlayerInputDebugger::Construct(const FArguments& InArgs)
+{
+	// Build the three tab content widgets.
+	SAssignNew(PlayerInputEventsTab, SPlayerInputEventsTab);
+
+	SAssignNew(DevicesTab, SInputDevicesTab);
+
+	SAssignNew(EnhancedInputTab, SEnhancedInputTab);
+
+	SAssignNew(GlobalInfoSection, SGlobalInputInfoSection);
+
+	TSharedRef<SWidgetSwitcher> Switcher = SAssignNew(TabSwitcher, SWidgetSwitcher)
+		+ SWidgetSwitcher::Slot() [ PlayerInputEventsTab.ToSharedRef() ]
+		+ SWidgetSwitcher::Slot() [ DevicesTab.ToSharedRef() ]
+		+ SWidgetSwitcher::Slot() [ EnhancedInputTab.ToSharedRef() ];
+
+	// Tab bar helper lambda
+	auto MakeTabButton = [this](int32 TabIndex, const FText& Label) -> TSharedRef<SWidget>
+	{
+		return SNew(SCheckBox)
+			.Style(FAppStyle::Get(), "ToggleButtonCheckbox")
+			.IsChecked_Lambda([this, TabIndex]()
+			{
+				return ActiveTabIndex == TabIndex ? ECheckBoxState::Checked : ECheckBoxState::Unchecked;
+			})
+			.OnCheckStateChanged_Lambda([this, TabIndex](ECheckBoxState NewState)
+			{
+				if (NewState == ECheckBoxState::Checked)
+				{
+					ActiveTabIndex = TabIndex;
+					TabSwitcher->SetActiveWidgetIndex(TabIndex);
+				}
+			})
+			.Padding(FMargin(10.f, 4.f))
+			[
+				SNew(STextBlock)
+				.Text(Label)
+			];
+	};
+
+	ChildSlot
+	[
+		SNew(SVerticalBox)
+
+		// Player controller selector
+		+ SVerticalBox::Slot()
+		.AutoHeight()
+		.Padding(6.f, 6.f, 6.f, 4.f)
+		[
+			SNew(SHorizontalBox)
+
+			+ SHorizontalBox::Slot()
+			.AutoWidth()
+			.VAlign(VAlign_Center)
+			.Padding(0.f, 0.f, 6.f, 0.f)
+			[
+				SNew(STextBlock)
+				.Text(LOCTEXT("PCLabel", "Player Controller:"))
+			]
+
+			+ SHorizontalBox::Slot()
+			.MaxWidth(320.f)
+			.AutoWidth()
+			[
+				SAssignNew(PCCombo, SComboBox<TSharedPtr<FPlayerControllerEntry>>)
+				.OptionsSource(&PCEntries)
+				.OnGenerateWidget(this, &SPlayerInputDebugger::MakePCComboEntry)
+				.OnSelectionChanged(this, &SPlayerInputDebugger::OnPCSelected)
+				.Content()
+				[
+					SNew(STextBlock)
+					.Text(this, &SPlayerInputDebugger::GetSelectedPCText)
+				]
+			]
+
+			+ SHorizontalBox::Slot()
+			.AutoWidth()
+			.VAlign(VAlign_Center)
+			.Padding(4.f, 0.f, 0.f, 0.f)
+			[
+				SNew(SButton)
+				.ButtonStyle(FAppStyle::Get(), "SimpleButton")
+				.ToolTipText(LOCTEXT("RefreshTooltip", "Refresh player controller list"))
+				.OnClicked_Lambda([this]() -> FReply
+				{
+					RefreshPlayerControllers();
+					return FReply::Handled();
+				})
+				[
+					SNew(SImage)
+					.Image(FAppStyle::GetBrush("Icons.Refresh"))
+					.DesiredSizeOverride(FVector2D(16.f, 16.f))
+				]
+			]
+		]
+
+		// Tab bar
+		+ SVerticalBox::Slot()
+		.AutoHeight()
+		[
+			SNew(SBorder)
+			.BorderImage(FAppStyle::GetBrush("ToolPanel.GroupBorder"))
+			.Padding(FMargin(0.f))
+			[
+				SNew(SHorizontalBox)
+
+				+ SHorizontalBox::Slot()
+				.AutoWidth()
+				[
+					MakeTabButton(0, LOCTEXT("TabInputEvents", "Input Events"))
+				]
+
+				+ SHorizontalBox::Slot()
+				.AutoWidth()
+				[
+					MakeTabButton(1, LOCTEXT("TabInputDevices", "Input Devices"))
+				]
+
+				+ SHorizontalBox::Slot()
+				.AutoWidth()
+				[
+					MakeTabButton(2, LOCTEXT("TabEnhancedInput", "Enhanced Input"))
+				]
+			]
+		]
+
+		// Global input overview (always visible, below tab bar)
+		+ SVerticalBox::Slot()
+		.AutoHeight()
+		[
+			GlobalInfoSection.ToSharedRef()
+		]
+
+		// Tab content
+		+ SVerticalBox::Slot()
+		.FillHeight(1.f)
+		[
+			Switcher
+		]
+	];
+
+	// BeginPIE clears stale entries; GameModePostLoginEvent repopulates as PCs are created.
+	BeginPIEHandle  = FEditorDelegates::BeginPIE.AddRaw(this, &SPlayerInputDebugger::OnBeginPIE);
+	EndPIEHandle    = FEditorDelegates::EndPIE.AddRaw(this, &SPlayerInputDebugger::OnEndPIE);
+	PostLoginHandle = FGameModeEvents::OnGameModePostLoginEvent().AddRaw(this, &SPlayerInputDebugger::OnGameModePostLogin);
+	LogoutHandle    = FGameModeEvents::OnGameModeLogoutEvent().AddRaw(this, &SPlayerInputDebugger::OnGameModeLogout);
+
+	// Populate if already in PIE.
+	RefreshPlayerControllers();
+}
+
+SPlayerInputDebugger::~SPlayerInputDebugger()
+{
+	FEditorDelegates::BeginPIE.Remove(BeginPIEHandle);
+	FEditorDelegates::EndPIE.Remove(EndPIEHandle);
+	FGameModeEvents::OnGameModePostLoginEvent().Remove(PostLoginHandle);
+	FGameModeEvents::OnGameModeLogoutEvent().Remove(LogoutHandle);
+}
+
+void SPlayerInputDebugger::RefreshPlayerControllers(const bool bIsPlayEnding /*= false*/)
+{
+	PCEntries.Empty();
+
+	if (!GEngine)
+	{
+		return;
+	}
+
+	if (!bIsPlayEnding)
+	{
+		int32 PCIndex = 0;
+		for (const FWorldContext& Context : GEngine->GetWorldContexts())
+		{
+			if (Context.WorldType != EWorldType::PIE)
+			{
+				continue;
+			}
+
+			UWorld* World = Context.World();
+			if (!World)
+			{
+				continue;
+			}
+
+			for (FConstPlayerControllerIterator It = World->GetPlayerControllerIterator(); It; ++It)
+			{
+				if (APlayerController* PC = It->Get())
+				{
+					TSharedPtr<FPlayerControllerEntry> Entry = MakeShared<FPlayerControllerEntry>();
+					Entry->PC = PC;
+					Entry->DisplayName = FText::Format(
+						LOCTEXT("PCEntryFmt", "Player Controller {0} ({1})"),
+						FText::AsNumber(PCIndex),
+						FText::FromString(PC->GetName()));
+					PCEntries.Add(MoveTemp(Entry));
+					++PCIndex;
+				}
+			}
+		}
+	}
+
+	if (PCCombo.IsValid())
+	{
+		PCCombo->RefreshOptions();
+	}
+
+	// Auto-select the first entry.
+	if (!PCEntries.IsEmpty())
+	{
+		SelectEntry(PCEntries[0]);
+	}
+	else
+	{
+		SelectEntry(nullptr);
+	}
+}
+
+void SPlayerInputDebugger::OnBeginPIE(const bool bIsSimulating)
+{
+	// Clear stale entries from any previous session.
+	// The list is repopulated by OnGameModePostLogin as player controllers are created.
+	PCEntries.Empty();
+	if (PCCombo.IsValid())
+	{
+		PCCombo->RefreshOptions();
+	}
+	SelectEntry(nullptr);
+}
+
+void SPlayerInputDebugger::OnGameModePostLogin(AGameModeBase* GameMode, APlayerController* NewPlayer)
+{
+	RefreshPlayerControllers();
+}
+
+void SPlayerInputDebugger::OnGameModeLogout(AGameModeBase* GameMode, AController* Exiting)
+{
+	// Guard against world teardown: GameModeLogoutEvent fires for each PC as the PIE
+	// world is destroyed, which can occur after EndPIE has already cleared the list.
+	// Refreshing during teardown would re-populate the dropdown from the dying world.
+	if (GEditor && GEditor->IsPlayingSessionInEditor())
+	{
+		//RefreshPlayerControllers();
+	}
+}
+
+void SPlayerInputDebugger::OnEndPIE(const bool bIsSimulating)
+{
+	PCEntries.Empty();
+	if (PCCombo.IsValid())
+	{
+		PCCombo->RefreshOptions();
+	}
+	SelectEntry(nullptr);
+
+	RefreshPlayerControllers(true);
+}
+
+void SPlayerInputDebugger::OnPCSelected(TSharedPtr<FPlayerControllerEntry> Entry, ESelectInfo::Type SelectInfo)
+{
+	SelectEntry(Entry);
+}
+
+void SPlayerInputDebugger::SelectEntry(TSharedPtr<FPlayerControllerEntry> Entry)
+{
+	SelectedEntry = Entry;
+	APlayerController* PC = Entry.IsValid() ? Entry->PC.Get() : nullptr;
+
+	if (EnhancedInputTab.IsValid())
+	{
+		EnhancedInputTab->SetPlayerController(PC);
+	}
+
+	if (DevicesTab.IsValid())
+	{
+		DevicesTab->SetPlayerController(PC);
+	}
+
+	if (PlayerInputEventsTab.IsValid())
+	{
+		PlayerInputEventsTab->SetPlayerController(PC);
+	}
+
+	if (GlobalInfoSection.IsValid())
+	{
+		GlobalInfoSection->SetPlayerController(PC);
+	}
+}
+
+TSharedRef<SWidget> SPlayerInputDebugger::MakePCComboEntry(TSharedPtr<FPlayerControllerEntry> Entry)
+{
+	return SNew(STextBlock)
+		.Text(Entry.IsValid() ? Entry->DisplayName : LOCTEXT("NullPC", "(none)"))
+		.Margin(FMargin(4.f, 2.f));
+}
+
+FText SPlayerInputDebugger::GetSelectedPCText() const
+{
+	if (SelectedEntry.IsValid() && SelectedEntry->PC.IsValid())
+	{
+		return SelectedEntry->DisplayName;
+	}
+	return LOCTEXT("NoPCSelected", "(No player controller — start PIE)");
+}
+
+#undef LOCTEXT_NAMESPACE
diff --git a/Engine/Plugins/PlayerInputDebugger/Source/PlayerInputDebugger/Private/SPlayerInputDebugger.h b/Engine/Plugins/PlayerInputDebugger/Source/PlayerInputDebugger/Private/SPlayerInputDebugger.h
new file mode 100644
--- /dev/null
+++ b/Engine/Plugins/PlayerInputDebugger/Source/PlayerInputDebugger/Private/SPlayerInputDebugger.h
@@ -0,0 +1,62 @@
+// Copyright Epic Games, Inc. All Rights Reserved.
+
+#pragma once
+
+#include "GameFramework/PlayerController.h"
+#include "Widgets/Input/SComboBox.h"
+#include "Widgets/Layout/SWidgetSwitcher.h"
+#include "Widgets/SCompoundWidget.h"
+
+class AController;
+class AGameModeBase;
+class SInputDevicesTab;
+class SEnhancedInputTab;
+class SPlayerInputEventsTab;
+class SGlobalInputInfoSection;
+
+// Entry in the player controller dropdown.
+struct FPlayerControllerEntry
+{
+	TWeakObjectPtr<APlayerController> PC;
+	FText DisplayName;
+};
+
+// Root widget for the Player Input Debugger window.
+// Contains a player controller selector and three sub-tabs.
+class SPlayerInputDebugger : public SCompoundWidget
+{
+public:
+	SLATE_BEGIN_ARGS(SPlayerInputDebugger) {}
+	SLATE_END_ARGS()
+
+	void Construct(const FArguments& InArgs);
+	virtual ~SPlayerInputDebugger();
+
+private:
+	void RefreshPlayerControllers(const bool bIsPlayEnding = false);
+	void OnBeginPIE(const bool bIsSimulating);
+	void OnEndPIE(const bool bIsSimulating);
+	void OnGameModePostLogin(AGameModeBase* GameMode, APlayerController* NewPlayer);
+	void OnGameModeLogout(AGameModeBase* GameMode, AController* Exiting);
+	void OnPCSelected(TSharedPtr<FPlayerControllerEntry> Entry, ESelectInfo::Type SelectInfo);
+	TSharedRef<SWidget> MakePCComboEntry(TSharedPtr<FPlayerControllerEntry> Entry);
+	FText GetSelectedPCText() const;
+	void SelectEntry(TSharedPtr<FPlayerControllerEntry> Entry);
+
+	TArray<TSharedPtr<FPlayerControllerEntry>> PCEntries;
+	TSharedPtr<FPlayerControllerEntry> SelectedEntry;
+
+	TSharedPtr<SComboBox<TSharedPtr<FPlayerControllerEntry>>> PCCombo;
+	TSharedPtr<SWidgetSwitcher> TabSwitcher;
+	int32 ActiveTabIndex = 0;
+
+	TSharedPtr<SPlayerInputEventsTab> PlayerInputEventsTab;
+	TSharedPtr<SInputDevicesTab> DevicesTab;
+	TSharedPtr<SEnhancedInputTab> EnhancedInputTab;
+	TSharedPtr<SGlobalInputInfoSection> GlobalInfoSection;
+
+	FDelegateHandle BeginPIEHandle;
+	FDelegateHandle EndPIEHandle;
+	FDelegateHandle PostLoginHandle;
+	FDelegateHandle LogoutHandle;
+};
diff --git a/Engine/Plugins/PlayerInputDebugger/Source/PlayerInputDebugger/Private/SPlayerInputEventsTab.cpp b/Engine/Plugins/PlayerInputDebugger/Source/PlayerInputDebugger/Private/SPlayerInputEventsTab.cpp
new file mode 100644
--- /dev/null
+++ b/Engine/Plugins/PlayerInputDebugger/Source/PlayerInputDebugger/Private/SPlayerInputEventsTab.cpp
@@ -0,0 +1,920 @@
+// Copyright Epic Games, Inc. All Rights Reserved.
+
+#include "SPlayerInputEventsTab.h"
+
+#include "Blueprint/UserWidget.h"
+#include "Brushes/SlateColorBrush.h"
+#include "Components/InputComponent.h"
+#include "Debugging/SlateDebugging.h"
+#include "Editor.h"
+#include "Engine/Engine.h"
+#include "Framework/Application/SlateApplication.h"
+#include "GameFramework/PlayerController.h"
+#include "GameFramework/PlayerInput.h"
+#include "Input/Events.h"
+#include "InputAction.h"
+#include "Slate/SObjectWidget.h"
+#include "Styling/AppStyle.h"
+#include "Styling/SlateTypes.h"
+#include "Subsystems/AssetEditorSubsystem.h"
+#include "Widgets/Input/SButton.h"
+#include "Widgets/Input/SCheckBox.h"
+#include "Widgets/Input/SHyperlink.h"
+#include "Widgets/Layout/SBorder.h"
+#include "Widgets/Layout/SBox.h"
+#include "Widgets/Layout/SScrollBox.h"
+#include "Widgets/Text/STextBlock.h"
+#include "Widgets/Views/SListView.h"
+#include "Widgets/Views/STableRow.h"
+
+#define LOCTEXT_NAMESPACE "SPlayerInputEventsTab"
+
+namespace InputEventsTab
+{
+	static const float FrameColumnWidth   = 70.f;
+	static const float SourceColumnWidth  = 80.f;
+	static const float EventColumnWidth   = 190.f;
+	static const float KeyColumnWidth     = 150.f;
+	static const float PCColumnWidth      = 120.f;
+	static const float HandlerColumnWidth = 210.f;
+}
+
+// ── Helpers ───────────────────────────────────────────────────────────────────
+
+// Slate row style: dark blue tint to distinguish from Player rows.
+static const FTableRowStyle& GetSlateRowStyle()
+{
+	static const FTableRowStyle Style = FTableRowStyle(FAppStyle::Get().GetWidgetStyle<FTableRowStyle>("TableView.Row"))
+		.SetEvenRowBackgroundBrush(FSlateColorBrush(FLinearColor(0.05f, 0.05f, 0.13f)))
+		.SetOddRowBackgroundBrush(FSlateColorBrush(FLinearColor(0.04f, 0.04f, 0.10f)));
+	return Style;
+}
+
+static FText GetSlateEventTypeText(ESlateDebuggingInputEvent Type)
+{
+	if (const UEnum* E = StaticEnum<ESlateDebuggingInputEvent>())
+	{
+		return E->GetDisplayNameTextByValue((int64)Type);
+	}
+	return FText::AsNumber((int32)Type);
+}
+
+// Compact value display: omit trailing zero axes.
+static FString FormatInputValue(const FVector& V)
+{
+	if (V.Z != 0.f) { return FString::Printf(TEXT("%.3f, %.3f, %.3f"), V.X, V.Y, V.Z); }
+	if (V.Y != 0.f) { return FString::Printf(TEXT("%.3f, %.3f"), V.X, V.Y); }
+	return FString::Printf(TEXT("%.3f"), V.X);
+}
+
+// ── Row widget ────────────────────────────────────────────────────────────────
+
+namespace
+{
+
+class SUnifiedInputEventRow : public SMultiColumnTableRow<TSharedPtr<FUnifiedInputEventRecord>>
+{
+public:
+	SLATE_BEGIN_ARGS(SUnifiedInputEventRow) {}
+		SLATE_ARGUMENT(TSharedPtr<FUnifiedInputEventRecord>, Item)
+	SLATE_END_ARGS()
+
+	void Construct(const FArguments& InArgs, const TSharedRef<STableViewBase>& InOwnerTableView)
+	{
+		Item = InArgs._Item;
+		FSuperRowType::FArguments SuperArgs = FSuperRowType::FArguments().Padding(FMargin(2.f, 1.f));
+		if (Item->Source == EInputEventSource::Slate || Item->Source == EInputEventSource::SlateMouseCapture)
+		{
+			SuperArgs.Style(&GetSlateRowStyle());
+		}
+		SMultiColumnTableRow<TSharedPtr<FUnifiedInputEventRecord>>::Construct(SuperArgs, InOwnerTableView);
+	}
+
+	virtual TSharedRef<SWidget> GenerateWidgetForColumn(const FName& ColumnName) override
+	{
+		check(Item.IsValid());
+
+		const bool bIsFlush        = (Item->Source == EInputEventSource::PlayerFlush);
+		const bool bIsMouseCapture = (Item->Source == EInputEventSource::SlateMouseCapture);
+		static const FLinearColor FlushColor(0.9f, 0.2f, 0.2f);
+		static const FLinearColor DimColor(0.5f, 0.5f, 0.5f);
+
+		// ── Frame ────────────────────────────────────────────────────────────
+		if (ColumnName == TEXT("Frame"))
+		{
+			return SNew(STextBlock)
+				.Text(FText::AsNumber(Item->FrameNumber))
+				.Font(FAppStyle::GetFontStyle(bIsFlush ? "BoldFont" : "NormalFont"))
+				.ColorAndOpacity(bIsFlush ? FlushColor : DimColor);
+		}
+
+		// ── Source ───────────────────────────────────────────────────────────
+		if (ColumnName == TEXT("Source"))
+		{
+			FText SourceText;
+			FLinearColor SourceColor;
+			switch (Item->Source)
+			{
+			case EInputEventSource::Slate:
+			case EInputEventSource::SlateMouseCapture:
+				SourceText  = LOCTEXT("SourceSlate",  "Slate");
+				SourceColor = FLinearColor(0.4f, 0.7f, 1.0f);   // blue
+				break;
+			case EInputEventSource::Player:
+				SourceText  = LOCTEXT("SourcePlayer", "Player");
+				SourceColor = FLinearColor(0.6f, 0.9f, 0.6f);   // green
+				break;
+			case EInputEventSource::PlayerFlush:
+				SourceText  = LOCTEXT("SourcePlayer", "Player");
+				SourceColor = FlushColor;
+				break;
+			}
+			return SNew(STextBlock)
+				.Text(SourceText)
+				.Font(FAppStyle::GetFontStyle("NormalFont"))
+				.ColorAndOpacity(SourceColor);
+		}
+
+		// ── Event / Action ───────────────────────────────────────────────────
+		if (ColumnName == TEXT("Event"))
+		{
+			if (bIsFlush)
+			{
+				return SNew(STextBlock)
+					.Text(LOCTEXT("FlushLabel", "--- Input Flushed ---"))
+					.Font(FAppStyle::GetFontStyle("BoldFont"))
+					.ColorAndOpacity(FlushColor);
+			}
+
+			if (bIsMouseCapture)
+			{
+				const bool bCaptured = Item->bReplyHandled;
+				return SNew(STextBlock)
+					.Text(bCaptured ? LOCTEXT("MouseCaptured", "Mouse Captured") : LOCTEXT("MouseCaptureLost", "Mouse Capture Lost"))
+					.Font(FAppStyle::GetFontStyle("NormalFont"))
+					.ColorAndOpacity(bCaptured ? FLinearColor(0.2f, 0.9f, 0.2f) : FLinearColor(0.9f, 0.6f, 0.2f));
+			}
+
+			if (Item->Source == EInputEventSource::Slate)
+			{
+				return SNew(STextBlock)
+					.Text(GetSlateEventTypeText(Item->SlateEventType))
+					.Font(FAppStyle::GetFontStyle("NormalFont"));
+			}
+
+			// Player — action name; hyperlink if we resolved an asset
+			if (Item->ActionAsset.IsValid())
+			{
+				return SNew(SHyperlink)
+					.Text(FText::FromString(Item->ActionName))
+					.OnNavigate(FSimpleDelegate::CreateLambda([WeakAsset = Item->ActionAsset]()
+					{
+						if (UObject* Asset = WeakAsset.Get())
+						{
+							if (UAssetEditorSubsystem* Sub = GEditor->GetEditorSubsystem<UAssetEditorSubsystem>())
+							{
+								Sub->OpenEditorForAsset(Asset);
+							}
+						}
+					}));
+			}
+			return SNew(STextBlock)
+				.Text(FText::FromString(Item->ActionName))
+				.Font(FAppStyle::GetFontStyle("NormalFont"));
+		}
+
+		// ── Key ──────────────────────────────────────────────────────────────
+		if (ColumnName == TEXT("Key"))
+		{
+			if (bIsFlush || bIsMouseCapture || Item->Source == EInputEventSource::Player)
+			{
+				return SNew(STextBlock).Text(FText::GetEmpty());
+			}
+			return SNew(STextBlock)
+				.Text(Item->Key.IsValid() ? FText::FromString(Item->Key.ToString()) : FText::GetEmpty())
+				.Font(FAppStyle::GetFontStyle("NormalFont"))
+				.ColorAndOpacity(FLinearColor(0.7f, 0.7f, 0.7f));
+		}
+
+		// ── Player Controller ─────────────────────────────────────────────────
+		if (ColumnName == TEXT("PC"))
+		{
+			return SNew(STextBlock)
+				.Text(FText::FromString(Item->PlayerControllerName))
+				.Font(FAppStyle::GetFontStyle("NormalFont"))
+				.ColorAndOpacity(FLinearColor(0.7f, 0.7f, 0.7f));
+		}
+
+		// ── Handler / Component ──────────────────────────────────────────────
+		if (ColumnName == TEXT("Handler"))
+		{
+			if (bIsFlush)
+			{
+				return SNew(STextBlock).Text(FText::GetEmpty());
+			}
+
+			if (bIsMouseCapture)
+			{
+				return SNew(STextBlock)
+					.Text(Item->HandlerWidgetType.IsEmpty() ? FText::GetEmpty() : FText::FromString(Item->HandlerWidgetType))
+					.Font(FAppStyle::GetFontStyle("NormalFont"))
+					.ColorAndOpacity(FLinearColor(0.9f, 0.9f, 0.2f));  // yellow
+			}
+
+			if (Item->Source == EInputEventSource::Slate)
+			{
+				if (Item->HandlerWidgetType.IsEmpty())
+				{
+					return SNew(STextBlock)
+						.Text(LOCTEXT("NoHandler", "—"))
+						.ColorAndOpacity(FLinearColor(0.4f, 0.4f, 0.4f));
+				}
+
+				FString WidgetLabel = Item->HandlerWidgetType;
+				if (!Item->UserWidgetName.IsEmpty())
+				{
+					WidgetLabel = FString::Printf(TEXT("%s (%s)"), *Item->HandlerWidgetType, *Item->UserWidgetName);
+				}
+
+				if (Item->UserWidgetBlueprint.IsValid())
+				{
+					static const FHyperlinkStyle BlueLinkStyle = []()
+					{
+						FHyperlinkStyle Style = FAppStyle::Get().GetWidgetStyle<FHyperlinkStyle>("Hyperlink");
+						Style.SetTextStyle(FTextBlockStyle(Style.TextStyle)
+							.SetColorAndOpacity(FSlateColor(FLinearColor(0.2f, 0.6f, 1.0f))));
+						return Style;
+					}();
+					TWeakObjectPtr<UObject> WeakBP = Item->UserWidgetBlueprint;
+					return SNew(SHyperlink)
+						.Text(FText::FromString(WidgetLabel))
+						.Style(&BlueLinkStyle)
+						.OnNavigate(FSimpleDelegate::CreateLambda([WeakBP]()
+						{
+							if (UObject* BP = WeakBP.Get())
+							{
+								if (UAssetEditorSubsystem* Sub = GEditor->GetEditorSubsystem<UAssetEditorSubsystem>())
+								{
+									Sub->OpenEditorForAsset(BP);
+								}
+							}
+						}));
+				}
+
+				return SNew(STextBlock)
+					.Text(FText::FromString(WidgetLabel))
+					.Font(FAppStyle::GetFontStyle("NormalFont"))
+					.ColorAndOpacity(FLinearColor(0.9f, 0.9f, 0.2f));  // yellow
+			}
+
+			// Player event — input component name
+			return SNew(STextBlock)
+				.Text(FText::FromString(Item->InputComponentName))
+				.Font(FAppStyle::GetFontStyle("NormalFont"));
+		}
+
+		return SNew(STextBlock).Text(FText::GetEmpty());
+	}
+
+private:
+	TSharedPtr<FUnifiedInputEventRecord> Item;
+};
+
+} // anonymous namespace
+
+// ── Construct / Destruct ──────────────────────────────────────────────────────
+
+void SPlayerInputEventsTab::Construct(const FArguments& InArgs)
+{
+	TSharedRef<SHeaderRow> HeaderRow = SNew(SHeaderRow)
+		+ SHeaderRow::Column("Frame")
+		.DefaultLabel(LOCTEXT("ColFrame", "Frame"))
+		.ManualWidth(InputEventsTab::FrameColumnWidth)
+
+		+ SHeaderRow::Column("Source")
+		.DefaultLabel(LOCTEXT("ColSource", "Source"))
+		.ManualWidth(InputEventsTab::SourceColumnWidth)
+
+		+ SHeaderRow::Column("Event")
+		.DefaultLabel(LOCTEXT("ColEvent", "Event / Action"))
+		.ManualWidth(InputEventsTab::EventColumnWidth)
+
+		+ SHeaderRow::Column("Key")
+		.DefaultLabel(LOCTEXT("ColKey", "Key"))
+		.ManualWidth(InputEventsTab::KeyColumnWidth)
+
+		+ SHeaderRow::Column("PC")
+		.DefaultLabel(LOCTEXT("ColPC", "Player Controller"))
+		.ManualWidth(InputEventsTab::PCColumnWidth)
+
+		+ SHeaderRow::Column("Handler")
+		.DefaultLabel(LOCTEXT("ColHandler", "Handler / Component"))
+		.ManualWidth(InputEventsTab::HandlerColumnWidth);
+
+	TSharedRef<SVerticalBox> VBox = SNew(SVerticalBox)
+
+		// Toolbar
+		+ SVerticalBox::Slot()
+		.AutoHeight()
+		.Padding(4.f, 4.f)
+		[
+			SNew(SHorizontalBox)
+
+			+ SHorizontalBox::Slot()
+			.AutoWidth()
+			.Padding(0.f, 0.f, 8.f, 0.f)
+			[
+				SNew(SButton)
+				.Text(LOCTEXT("ClearBtn", "Clear"))
+				.OnClicked_Lambda([this]() -> FReply
+				{
+					LogItems.Empty();
+					SelectedItem.Reset();
+					ListView->ClearSelection();
+					RebuildDetails();
+					ListView->RequestListRefresh();
+					return FReply::Handled();
+				})
+			]
+
+			+ SHorizontalBox::Slot()
+			.AutoWidth()
+			.Padding(0.f, 0.f, 8.f, 0.f)
+			[
+				SNew(SButton)
+				.Text(LOCTEXT("ResumeScrollBtn", "Resume Scroll"))
+				.IsEnabled_Lambda([this]() { return SelectedItem.IsValid(); })
+				.ToolTipText(LOCTEXT("ResumeScrollTip", "Clear the current selection and resume auto-scrolling to new events"))
+				.OnClicked_Lambda([this]() -> FReply
+				{
+					SelectedItem.Reset();
+					ListView->ClearSelection();
+					RebuildDetails();
+					if (bAutoScroll && !LogItems.IsEmpty())
+					{
+						ListView->RequestScrollIntoView(LogItems.Last());
+					}
+					return FReply::Handled();
+				})
+			]
+
+			+ SHorizontalBox::Slot()
+			.AutoWidth()
+			.VAlign(VAlign_Center)
+			.Padding(0.f, 0.f, 8.f, 0.f)
+			[
+				SNew(SCheckBox)
+				.IsChecked_Lambda([this]() { return bAutoScroll ? ECheckBoxState::Checked : ECheckBoxState::Unchecked; })
+				.OnCheckStateChanged_Lambda([this](ECheckBoxState State)
+				{
+					bAutoScroll = (State == ECheckBoxState::Checked);
+				})
+				[
+					SNew(STextBlock).Text(LOCTEXT("AutoScrollLabel", "Auto-scroll"))
+				]
+			]
+
+			+ SHorizontalBox::Slot()
+			.AutoWidth()
+			.VAlign(VAlign_Center)
+			.Padding(0.f, 0.f, 8.f, 0.f)
+			[
+				SNew(SCheckBox)
+				.IsChecked_Lambda([this]() { return bIgnoreDebuggerWindowEvents ? ECheckBoxState::Checked : ECheckBoxState::Unchecked; })
+				.OnCheckStateChanged_Lambda([this](ECheckBoxState State)
+				{
+					bIgnoreDebuggerWindowEvents = (State == ECheckBoxState::Checked);
+				})
+				[
+					SNew(STextBlock).Text(LOCTEXT("IgnoreDebuggerLabel", "Ignore events on this window"))
+				]
+			]
+
+			+ SHorizontalBox::Slot()
+			.AutoWidth()
+			.VAlign(VAlign_Center)
+			[
+				SNew(SCheckBox)
+				.IsChecked_Lambda([this]() { return bFilterZeroLegacyInputs ? ECheckBoxState::Checked : ECheckBoxState::Unchecked; })
+				.OnCheckStateChanged_Lambda([this](ECheckBoxState State)
+				{
+					bFilterZeroLegacyInputs = (State == ECheckBoxState::Checked);
+				})
+				[
+					SNew(STextBlock).Text(LOCTEXT("FilterZeroLegacyLabel", "Filters Legacy Inputs"))
+				]
+			]
+
+			+ SHorizontalBox::Slot()
+			.AutoWidth()
+			.VAlign(VAlign_Center)
+			.Padding(8.f, 0.f, 0.f, 0.f)
+			[
+				SNew(SCheckBox)
+				.IsChecked_Lambda([this]() { return bShowSlateMouseFocusEvents ? ECheckBoxState::Checked : ECheckBoxState::Unchecked; })
+				.OnCheckStateChanged_Lambda([this](ECheckBoxState State)
+				{
+					bShowSlateMouseFocusEvents = (State == ECheckBoxState::Checked);
+				})
+				[
+					SNew(STextBlock).Text(LOCTEXT("FilterMouseFocusLabel", "Show Mouse Focus Events"))
+				]
+			]
+		];
+
+#if !UE_PLAYER_INPUT_INCLUDE_DEBUG
+	VBox->AddSlot()
+	.AutoHeight()
+	.Padding(8.f, 4.f)
+	[
+		SNew(STextBlock)
+		.Text(LOCTEXT("PlayerEventsNotAvailable",
+			"Player input events are not available in this configuration. "
+			"(Requires UE_PLAYER_INPUT_INCLUDE_DEBUG=1)"))
+		.ColorAndOpacity(FLinearColor(0.9f, 0.6f, 0.2f))
+		.AutoWrapText(true)
+	];
+#endif
+
+	VBox->AddSlot()
+	.FillHeight(1.f)
+	[
+		SAssignNew(ListView, SListView<TSharedPtr<FUnifiedInputEventRecord>>)
+		.ListItemsSource(&LogItems)
+		.OnGenerateRow(this, &SPlayerInputEventsTab::OnGenerateRow)
+		.HeaderRow(HeaderRow)
+		.SelectionMode(ESelectionMode::Single)
+		.OnSelectionChanged_Lambda([this](TSharedPtr<FUnifiedInputEventRecord> Item, ESelectInfo::Type)
+		{
+			SelectedItem = Item;
+			RebuildDetails();
+		})
+	];
+
+	VBox->AddSlot()
+	.AutoHeight()
+	[
+		SNew(SBorder)
+		.BorderImage(FAppStyle::GetBrush("ToolPanel.GroupBorder"))
+		[
+			SNew(SBox)
+			.MaxDesiredHeight(200.f)
+			[
+				SAssignNew(DetailsBox, SBox)
+			]
+		]
+	];
+
+	ChildSlot[ VBox ];
+
+	RebuildDetails();
+
+	// Store a weak reference to ourselves so we can identify our own window.
+	WeakSelfWidget = SharedThis(this);
+
+	BeginPIEHandle = FEditorDelegates::BeginPIE.AddRaw(this, &SPlayerInputEventsTab::OnBeginPIE);
+
+#if WITH_SLATE_DEBUGGING
+	SlateInputEventHandle    = FSlateDebugging::InputEvent.AddRaw(this, &SPlayerInputEventsTab::OnSlateInputEvent);
+	SlateMouseCaptureHandle  = FSlateDebugging::MouseCaptureEvent.AddRaw(this, &SPlayerInputEventsTab::OnSlateMouseCaptureEvent);
+#endif
+
+#if UE_PLAYER_INPUT_INCLUDE_DEBUG
+	EventExecutedHandle = UE::Input::FPlayerInputDebugging::OnPlayerInputEventExecuted.AddRaw(
+		this, &SPlayerInputEventsTab::OnPlayerInputEventExecuted);
+	FlushedHandle = UE::Input::FPlayerInputDebugging::OnPlayerInputFlushed.AddRaw(
+		this, &SPlayerInputEventsTab::OnPlayerInputFlushed);
+#endif
+}
+
+SPlayerInputEventsTab::~SPlayerInputEventsTab()
+{
+	FEditorDelegates::BeginPIE.Remove(BeginPIEHandle);
+
+#if WITH_SLATE_DEBUGGING
+	FSlateDebugging::InputEvent.Remove(SlateInputEventHandle);
+	FSlateDebugging::MouseCaptureEvent.Remove(SlateMouseCaptureHandle);
+#endif
+
+#if UE_PLAYER_INPUT_INCLUDE_DEBUG
+	UE::Input::FPlayerInputDebugging::OnPlayerInputEventExecuted.Remove(EventExecutedHandle);
+	UE::Input::FPlayerInputDebugging::OnPlayerInputFlushed.Remove(FlushedHandle);
+#endif
+}
+
+void SPlayerInputEventsTab::SetPlayerController(APlayerController* /*PC*/)
+{
+	// Player controller selection no longer filters events; all PCs are shown.
+}
+
+void SPlayerInputEventsTab::OnBeginPIE(const bool /*bIsSimulating*/)
+{
+	LogItems.Empty();
+	SelectedItem.Reset();
+	RebuildDetails();
+	ListView->RequestListRefresh();
+}
+
+// ── AppendRecord ──────────────────────────────────────────────────────────────
+
+void SPlayerInputEventsTab::AppendRecord(TSharedPtr<FUnifiedInputEventRecord> Record)
+{
+	if (LogItems.Num() >= MaxLogEntries)
+	{
+		LogItems.RemoveAt(0, EAllowShrinking::No);
+	}
+	LogItems.Add(MoveTemp(Record));
+	ListView->RequestListRefresh();
+
+	if (bAutoScroll && !SelectedItem.IsValid() && !LogItems.IsEmpty())
+	{
+		ListView->RequestScrollIntoView(LogItems.Last());
+	}
+}
+
+// ── Row generation ────────────────────────────────────────────────────────────
+
+TSharedRef<ITableRow> SPlayerInputEventsTab::OnGenerateRow(TSharedPtr<FUnifiedInputEventRecord> Item, const TSharedRef<STableViewBase>& OwnerTable)
+{
+	return SNew(SUnifiedInputEventRow, OwnerTable)
+		.Item(Item);
+}
+
+// ── Slate input event handler ─────────────────────────────────────────────────
+
+#if WITH_SLATE_DEBUGGING
+
+void SPlayerInputEventsTab::OnSlateInputEvent(const FSlateDebuggingInputEventArgs& EventArgs)
+{
+	// Always skip high-frequency mouse move events.
+	if (EventArgs.InputEventType == ESlateDebuggingInputEvent::MouseMove)
+	{
+		return;
+	}
+
+	// Optionally filter mouse enter/leave (focus) events.
+	if (!bShowSlateMouseFocusEvents &&
+		(EventArgs.InputEventType == ESlateDebuggingInputEvent::MouseEnter ||
+		 EventArgs.InputEventType == ESlateDebuggingInputEvent::MouseLeave))
+	{
+		return;
+	}
+
+	// Filter events that originate inside our own debugger window.
+	if (bIgnoreDebuggerWindowEvents && EventArgs.HandlerWidget.IsValid())
+	{
+		if (TSharedPtr<SWidget> Self = WeakSelfWidget.Pin())
+		{
+			if (FSlateApplication::IsInitialized())
+			{
+				TSharedPtr<SWindow> HandlerWindow = FSlateApplication::Get().FindWidgetWindow(EventArgs.HandlerWidget.ToSharedRef());
+				TSharedPtr<SWindow> SelfWindow    = FSlateApplication::Get().FindWidgetWindow(Self.ToSharedRef());
+				if (HandlerWindow.IsValid() && HandlerWindow == SelfWindow)
+				{
+					return;
+				}
+			}
+		}
+	}
+
+	TSharedPtr<FUnifiedInputEventRecord> Record = MakeShared<FUnifiedInputEventRecord>();
+	Record->FrameNumber    = GFrameNumber;
+	Record->Source         = EInputEventSource::Slate;
+	Record->SlateEventType = EventArgs.InputEventType;
+	Record->bReplyHandled  = EventArgs.Reply.IsEventHandled();
+	Record->AdditionalContent = EventArgs.AdditionalContent;
+
+	// Extract key and analog value from the typed FInputEvent pointer.
+	if (const FInputEvent* IE = EventArgs.InputEvent)
+	{
+		switch (EventArgs.InputEventType)
+		{
+		case ESlateDebuggingInputEvent::PreviewKeyDown:
+		case ESlateDebuggingInputEvent::KeyDown:
+		case ESlateDebuggingInputEvent::KeyUp:
+			{
+				const FKeyEvent* KE = static_cast<const FKeyEvent*>(IE);
+				Record->Key      = KE->GetKey();
+				Record->RawValue = (EventArgs.InputEventType == ESlateDebuggingInputEvent::KeyUp) ? 0.f : 1.f;
+			}
+			break;
+
+		case ESlateDebuggingInputEvent::AnalogInput:
+			{
+				const FAnalogInputEvent* AE = static_cast<const FAnalogInputEvent*>(IE);
+				Record->Key      = AE->GetKey();
+				Record->RawValue = AE->GetAnalogValue();
+			}
+			break;
+
+		case ESlateDebuggingInputEvent::PreviewMouseButtonDown:
+		case ESlateDebuggingInputEvent::MouseButtonDown:
+		case ESlateDebuggingInputEvent::MouseButtonUp:
+		case ESlateDebuggingInputEvent::MouseButtonDoubleClick:
+			{
+				const FPointerEvent* PE = static_cast<const FPointerEvent*>(IE);
+				Record->Key      = PE->GetEffectingButton();
+				Record->RawValue = (EventArgs.InputEventType == ESlateDebuggingInputEvent::MouseButtonUp) ? 0.f : 1.f;
+			}
+			break;
+
+		case ESlateDebuggingInputEvent::MouseWheel:
+			{
+				const FPointerEvent* PE = static_cast<const FPointerEvent*>(IE);
+				Record->RawValue = PE->GetWheelDelta();
+			}
+			break;
+
+		default:
+			break;
+		}
+	}
+
+	// Handler widget info: type name + walk up to find the nearest UUserWidget.
+	if (EventArgs.HandlerWidget.IsValid())
+	{
+		Record->HandlerWidgetType = EventArgs.HandlerWidget->GetType().ToString();
+
+		TSharedPtr<SWidget> Current = EventArgs.HandlerWidget;
+		while (Current.IsValid())
+		{
+			if (Current->GetType() == TEXT("SObjectWidget"))
+			{
+				if (UUserWidget* UW = static_cast<SObjectWidget*>(Current.Get())->GetWidgetObject())
+				{
+					Record->UserWidgetName      = UW->GetClass()->GetName();
+					Record->UserWidgetBlueprint = UW->GetClass()->ClassGeneratedBy;
+				}
+				break;
+			}
+			Current = Current->GetParentWidget();
+		}
+	}
+
+	AppendRecord(MoveTemp(Record));
+}
+
+void SPlayerInputEventsTab::OnSlateMouseCaptureEvent(const FSlateDebuggingMouseCaptureEventArgs& EventArgs)
+{
+	// Filter events originating in our own debugger window.
+	if (bIgnoreDebuggerWindowEvents && EventArgs.CaptureWidget.IsValid())
+	{
+		if (TSharedPtr<SWidget> Self = WeakSelfWidget.Pin())
+		{
+			if (FSlateApplication::IsInitialized())
+			{
+				TSharedPtr<SWindow> CaptureWindow = FSlateApplication::Get().FindWidgetWindow(
+					ConstCastSharedRef<SWidget>(EventArgs.CaptureWidget.ToSharedRef()));
+				TSharedPtr<SWindow> SelfWindow = FSlateApplication::Get().FindWidgetWindow(Self.ToSharedRef());
+				if (CaptureWindow.IsValid() && CaptureWindow == SelfWindow)
+				{
+					return;
+				}
+			}
+		}
+	}
+
+	TSharedPtr<FUnifiedInputEventRecord> Record = MakeShared<FUnifiedInputEventRecord>();
+	Record->FrameNumber   = GFrameNumber;
+	Record->Source        = EInputEventSource::SlateMouseCapture;
+	Record->bReplyHandled = EventArgs.Captured; // true = captured, false = lost
+
+	if (EventArgs.CaptureWidget.IsValid())
+	{
+		Record->HandlerWidgetType = EventArgs.CaptureWidget->GetType().ToString();
+	}
+
+	AppendRecord(MoveTemp(Record));
+}
+
+#endif  // WITH_SLATE_DEBUGGING
+
+// ── Player input event handlers ───────────────────────────────────────────────
+
+#if UE_PLAYER_INPUT_INCLUDE_DEBUG
+
+void SPlayerInputEventsTab::OnPlayerInputEventExecuted(const UPlayerInput* PlayerInput, const UE::Input::FPlayerInputDebuggingArgs& Args)
+{
+	// Resolve the action name to a UInputAction asset (only succeeds for Enhanced Input actions).
+	UInputAction* Action = FindObject<UInputAction>(nullptr, *Args.ActionName);
+
+	// When filtering is enabled, drop legacy (non-Enhanced-Input) events whose value is zero.
+	// These represent un-pressed bindings that fire every tick and clutter the log.
+	if (bFilterZeroLegacyInputs && !Action && Args.InputValue.IsNearlyZero())
+	{
+		return;
+	}
+
+	TSharedPtr<FUnifiedInputEventRecord> Record = MakeShared<FUnifiedInputEventRecord>();
+	Record->FrameNumber = GFrameNumber;
+	Record->Source      = EInputEventSource::Player;
+	Record->InputValue  = Args.InputValue;
+
+	// Identify the owning player controller so the PC column can be populated.
+	if (const APlayerController* OwningPC = PlayerInput->GetTypedOuter<APlayerController>())
+	{
+		Record->PlayerControllerName = OwningPC->GetName();
+	}
+
+	if (Action)
+	{
+		Record->ActionAsset = Action;
+		Record->ActionName  = Action->GetName();
+	}
+	else
+	{
+		Record->ActionName = Args.ActionName;
+	}
+
+	if (const UInputComponent* IC = Args.InputComponent.Get())
+	{
+		Record->InputComponentName     = IC->GetName();
+		Record->InputComponentPriority = IC->Priority;
+		if (const AActor* Owner = IC->GetOwner())
+		{
+			Record->OwnerName      = Owner->GetActorNameOrLabel();
+			Record->OwnerBlueprint = Owner->GetClass()->ClassGeneratedBy;
+		}
+		else if (const UObject* Outer = IC->GetOuter())
+		{
+			Record->OwnerName      = Outer->GetName();
+			Record->OwnerBlueprint = Outer->GetClass()->ClassGeneratedBy;
+		}
+	}
+
+	if (const UObject* Listener = Args.ListeningObject.Get())
+	{
+		if (const AActor* ListenerActor = Cast<AActor>(Listener))
+		{
+			Record->ListeningObjectName = ListenerActor->GetActorNameOrLabel();
+		}
+		else
+		{
+			Record->ListeningObjectName = Listener->GetName();
+		}
+		Record->ListeningObjectBlueprint = Listener->GetClass()->ClassGeneratedBy;
+	}
+
+	AppendRecord(MoveTemp(Record));
+}
+
+void SPlayerInputEventsTab::OnPlayerInputFlushed(const UPlayerInput* PlayerInput)
+{
+	TSharedPtr<FUnifiedInputEventRecord> Record = MakeShared<FUnifiedInputEventRecord>();
+	Record->FrameNumber = GFrameNumber;
+	Record->Source      = EInputEventSource::PlayerFlush;
+
+	if (const APlayerController* OwningPC = PlayerInput->GetTypedOuter<APlayerController>())
+	{
+		Record->PlayerControllerName = OwningPC->GetName();
+	}
+
+	AppendRecord(MoveTemp(Record));
+}
+
+#endif  // UE_PLAYER_INPUT_INCLUDE_DEBUG
+
+// ── Details panel ─────────────────────────────────────────────────────────────
+
+void SPlayerInputEventsTab::RebuildDetails()
+{
+	if (!SelectedItem.IsValid())
+	{
+		DetailsBox->SetContent(
+			SNew(SBox).Padding(FMargin(8.f))
+			[
+				SNew(STextBlock)
+				.Text(LOCTEXT("NoSelection", "Select a row to view details."))
+				.ColorAndOpacity(FLinearColor(0.5f, 0.5f, 0.5f))
+			]
+		);
+		return;
+	}
+
+	if (SelectedItem->Source == EInputEventSource::PlayerFlush)
+	{
+		DetailsBox->SetContent(
+			SNew(SBox).Padding(FMargin(8.f))
+			[
+				SNew(STextBlock)
+				.Text(FText::Format(
+					LOCTEXT("FlushDetails", "Input Flush — Frame {0}  ({1})"),
+					FText::AsNumber(SelectedItem->FrameNumber),
+					FText::FromString(SelectedItem->PlayerControllerName)))
+				.ColorAndOpacity(FLinearColor(0.9f, 0.2f, 0.2f))
+			]
+		);
+		return;
+	}
+
+	TSharedPtr<FUnifiedInputEventRecord> Item = SelectedItem;
+
+	// Returns a hyperlink if the object has a blueprint asset, otherwise plain text.
+	auto MakeObjectWidget = [](const FString& Name, TWeakObjectPtr<UObject> Blueprint) -> TSharedRef<SWidget>
+	{
+		if (Blueprint.IsValid())
+		{
+			return SNew(SHyperlink)
+				.Text(FText::FromString(Name))
+				.OnNavigate(FSimpleDelegate::CreateLambda([Blueprint]()
+				{
+					if (UObject* BP = Blueprint.Get())
+					{
+						if (UAssetEditorSubsystem* Sub = GEditor->GetEditorSubsystem<UAssetEditorSubsystem>())
+						{
+							Sub->OpenEditorForAsset(BP);
+						}
+					}
+				}));
+		}
+		return SNew(STextBlock).Text(FText::FromString(Name));
+	};
+
+	auto MakeRow = [](const FText& Label, TSharedRef<SWidget> Value) -> TSharedRef<SWidget>
+	{
+		return SNew(SHorizontalBox)
+			+ SHorizontalBox::Slot().AutoWidth().VAlign(VAlign_Center)
+			[
+				SNew(STextBlock)
+				.Text(Label)
+				.ColorAndOpacity(FLinearColor(0.6f, 0.6f, 0.6f))
+				.MinDesiredWidth(170.f)
+			]
+			+ SHorizontalBox::Slot().AutoWidth().VAlign(VAlign_Center)
+			[
+				Value
+			];
+	};
+
+	TSharedRef<SScrollBox> ScrollBox = SNew(SScrollBox);
+	auto AddRow = [&](const FText& Label, TSharedRef<SWidget> Value)
+	{
+		ScrollBox->AddSlot().Padding(FMargin(8.f, 3.f))[ MakeRow(Label, Value) ];
+	};
+
+	AddRow(LOCTEXT("DetailFrame", "Frame:"),
+		SNew(STextBlock).Text(FText::AsNumber(Item->FrameNumber)));
+
+	if (Item->Source == EInputEventSource::Slate)
+	{
+		AddRow(LOCTEXT("DetailSource",     "Source:"),
+			SNew(STextBlock).Text(LOCTEXT("SourceSlate", "Slate"))
+				.ColorAndOpacity(FLinearColor(0.4f, 0.7f, 1.0f)));
+		AddRow(LOCTEXT("DetailEventType",  "Event Type:"),
+			SNew(STextBlock).Text(GetSlateEventTypeText(Item->SlateEventType)));
+		AddRow(LOCTEXT("DetailKey",        "Key:"),
+			SNew(STextBlock).Text(Item->Key.IsValid() ? FText::FromString(Item->Key.ToString()) : FText::GetEmpty()));
+		AddRow(LOCTEXT("DetailValue",      "Value:"),
+			SNew(STextBlock).Text(FText::FromString(FString::SanitizeFloat(Item->RawValue))));
+		AddRow(LOCTEXT("DetailHandler",    "Handler Widget:"),
+			MakeObjectWidget(Item->UserWidgetName.IsEmpty()
+				? Item->HandlerWidgetType
+				: FString::Printf(TEXT("%s (%s)"), *Item->HandlerWidgetType, *Item->UserWidgetName),
+				Item->UserWidgetBlueprint));
+		AddRow(LOCTEXT("DetailReply",      "Reply Handled:"),
+			SNew(STextBlock).Text(Item->bReplyHandled ? LOCTEXT("Yes","Yes") : LOCTEXT("No","No"))
+				.ColorAndOpacity(Item->bReplyHandled ? FLinearColor(0.2f,0.9f,0.2f) : FLinearColor(0.5f,0.5f,0.5f)));
+		if (!Item->AdditionalContent.IsEmpty())
+		{
+			AddRow(LOCTEXT("DetailAdditional", "Additional:"),
+				SNew(STextBlock).Text(FText::FromString(Item->AdditionalContent)));
+		}
+	}
+	else if (Item->Source == EInputEventSource::SlateMouseCapture)
+	{
+		const bool bCaptured = Item->bReplyHandled;
+		AddRow(LOCTEXT("DetailSource",        "Source:"),
+			SNew(STextBlock).Text(LOCTEXT("SourceSlate", "Slate"))
+				.ColorAndOpacity(FLinearColor(0.4f, 0.7f, 1.0f)));
+		AddRow(LOCTEXT("DetailMouseCapEvent", "Event:"),
+			SNew(STextBlock)
+				.Text(bCaptured ? LOCTEXT("MouseCaptured", "Mouse Captured") : LOCTEXT("MouseCaptureLost", "Mouse Capture Lost"))
+				.ColorAndOpacity(bCaptured ? FLinearColor(0.2f, 0.9f, 0.2f) : FLinearColor(0.9f, 0.6f, 0.2f)));
+		AddRow(LOCTEXT("DetailCapWidget",     "Capturing Widget:"),
+			SNew(STextBlock).Text(FText::FromString(Item->HandlerWidgetType)));
+	}
+	else
+	{
+		AddRow(LOCTEXT("DetailSource",     "Source:"),
+			SNew(STextBlock).Text(LOCTEXT("SourcePlayer", "Player"))
+				.ColorAndOpacity(FLinearColor(0.6f, 0.9f, 0.6f)));
+		AddRow(LOCTEXT("DetailPC",         "Player Controller:"),
+			SNew(STextBlock).Text(FText::FromString(Item->PlayerControllerName)));
+		AddRow(LOCTEXT("DetailComponent",  "Input Component:"),
+			SNew(STextBlock).Text(FText::FromString(Item->InputComponentName)));
+		AddRow(LOCTEXT("DetailPriority",   "Priority:"),
+			SNew(STextBlock).Text(FText::AsNumber(Item->InputComponentPriority)));
+		AddRow(LOCTEXT("DetailOwner",      "Owner:"),
+			MakeObjectWidget(Item->OwnerName, Item->OwnerBlueprint));
+		AddRow(LOCTEXT("DetailListener",   "Listening Object:"),
+			MakeObjectWidget(Item->ListeningObjectName, Item->ListeningObjectBlueprint));
+		AddRow(LOCTEXT("DetailAction",     "Action:"),
+			MakeObjectWidget(Item->ActionName, Item->ActionAsset));
+		AddRow(LOCTEXT("DetailValue",      "Value:"),
+			SNew(STextBlock).Text(FText::FromString(FormatInputValue(Item->InputValue))));
+	}
+
+	DetailsBox->SetContent(ScrollBox);
+}
+
+#undef LOCTEXT_NAMESPACE
diff --git a/Engine/Plugins/PlayerInputDebugger/Source/PlayerInputDebugger/Private/SPlayerInputEventsTab.h b/Engine/Plugins/PlayerInputDebugger/Source/PlayerInputDebugger/Private/SPlayerInputEventsTab.h
new file mode 100644
--- /dev/null
+++ b/Engine/Plugins/PlayerInputDebugger/Source/PlayerInputDebugger/Private/SPlayerInputEventsTab.h
@@ -0,0 +1,98 @@
+// Copyright Epic Games, Inc. All Rights Reserved.
+
+#pragma once
+
+#include "Debugging/SlateDebugging.h"
+#include "GameFramework/PlayerController.h"
+#include "GameFramework/PlayerInput.h"
+#include "Widgets/SCompoundWidget.h"
+#include "Widgets/Views/ITableRow.h"
+#include "Widgets/Views/SListView.h"
+#include "GameFramework/PlayerInputDebugging.h"
+
+// Discriminates each log entry in the combined Input Events tab.
+enum class EInputEventSource : uint8
+{
+	Slate,             // from FSlateDebugging::InputEvent
+	SlateMouseCapture, // from FSlateDebugging::MouseCaptureEvent
+	Player,            // from FPlayerInputDebugging::OnPlayerInputEventExecuted
+	PlayerFlush,       // flush marker from FPlayerInputDebugging::OnPlayerInputFlushed
+};
+
+// A single entry in the combined input event log.
+struct FUnifiedInputEventRecord
+{
+	uint32 FrameNumber = 0;
+	EInputEventSource Source = EInputEventSource::Slate;
+
+	// ── Slate fields ─────────────────────────────────────────────────────────
+	ESlateDebuggingInputEvent SlateEventType = ESlateDebuggingInputEvent::MAX;
+	FKey Key;
+	float RawValue = 0.f;
+	FString HandlerWidgetType;
+	FString UserWidgetName;                      // nearest UUserWidget class name, if any
+	TWeakObjectPtr<UObject> UserWidgetBlueprint; // ClassGeneratedBy for hyperlink
+	FString AdditionalContent;
+	bool bReplyHandled = false; // Slate: reply handled; SlateMouseCapture: true=captured, false=lost
+
+	// ── Player fields ─────────────────────────────────────────────────────────
+	FString PlayerControllerName; // PC that owns the PlayerInput; empty for Slate events
+	FString ActionName;
+	TWeakObjectPtr<UObject> ActionAsset;
+	FString InputComponentName;
+	int32 InputComponentPriority = 0;
+	FString OwnerName;
+	TWeakObjectPtr<UObject> OwnerBlueprint;
+	FString ListeningObjectName;
+	TWeakObjectPtr<UObject> ListeningObjectBlueprint;
+	FVector InputValue = FVector::ZeroVector;
+};
+
+// Combined Input Events tab: Slate input events (FSlateDebugging::InputEvent) and
+// Player input events (FPlayerInputDebugging) in a single clickable log.
+class SPlayerInputEventsTab : public SCompoundWidget
+{
+public:
+	SLATE_BEGIN_ARGS(SPlayerInputEventsTab) {}
+	SLATE_END_ARGS()
+
+	void Construct(const FArguments& InArgs);
+	virtual ~SPlayerInputEventsTab();
+
+	void SetPlayerController(APlayerController* PC);
+
+private:
+	TSharedRef<ITableRow> OnGenerateRow(TSharedPtr<FUnifiedInputEventRecord> Item, const TSharedRef<STableViewBase>& OwnerTable);
+	void AppendRecord(TSharedPtr<FUnifiedInputEventRecord> Record);
+	void RebuildDetails();
+	void OnBeginPIE(const bool bIsSimulating);
+
+#if WITH_SLATE_DEBUGGING
+	void OnSlateInputEvent(const FSlateDebuggingInputEventArgs& EventArgs);
+	void OnSlateMouseCaptureEvent(const FSlateDebuggingMouseCaptureEventArgs& EventArgs);
+	FDelegateHandle SlateInputEventHandle;
+	FDelegateHandle SlateMouseCaptureHandle;
+#endif
+
+#if UE_PLAYER_INPUT_INCLUDE_DEBUG
+	void OnPlayerInputEventExecuted(const UPlayerInput* PlayerInput, const UE::Input::FPlayerInputDebuggingArgs& Args);
+	void OnPlayerInputFlushed(const UPlayerInput* PlayerInput);
+	FDelegateHandle EventExecutedHandle;
+	FDelegateHandle FlushedHandle;
+#endif
+
+	FDelegateHandle BeginPIEHandle;
+	TWeakPtr<SWidget> WeakSelfWidget; // used to identify our window for event filtering
+
+	TArray<TSharedPtr<FUnifiedInputEventRecord>> LogItems;
+	TSharedPtr<SListView<TSharedPtr<FUnifiedInputEventRecord>>> ListView;
+	bool bAutoScroll = true;
+	bool bIgnoreDebuggerWindowEvents = true;
+	bool bFilterZeroLegacyInputs = true;
+	bool bShowSlateMouseFocusEvents = false;
+
+	TSharedPtr<FUnifiedInputEventRecord> SelectedItem;
+	TSharedPtr<SBox> DetailsBox;
+
+	static constexpr int32 MaxLogEntries = 1024;
+};
