[Player input debugger]

- Change the tool menu location to Tools > Debugging > Player Input Debugger, to be with everything else
- Gave it a nice icon
- Add info about the IMC filter modes
- Make the Enhanced input tabs clickable with their asset links

#rnx
#Synced-CL 51542463

diff --git a/Engine/Plugins/EnhancedInput/Source/EnhancedInput/Public/InputMappingContext.h b/Engine/Plugins/EnhancedInput/Source/EnhancedInput/Public/InputMappingContext.h
--- a/Engine/Plugins/EnhancedInput/Source/EnhancedInput/Public/InputMappingContext.h
+++ b/Engine/Plugins/EnhancedInput/Source/EnhancedInput/Public/InputMappingContext.h
@@ -171,6 +171,9 @@ public:
 	 * @return The registration tracking mode that this IMC is using. 
 	 */
 	EMappingContextRegistrationTrackingMode GetRegistrationTrackingMode() const { return RegistrationTrackingMode; }
+
+	/** @return The input mode filter options set on this mapping context. */
+	EMappingContextInputModeFilterOptions GetInputModeFilterOptions() const { return InputModeFilterOptions; }
 	
 	// Localized context descriptor
 	UPROPERTY(BlueprintReadOnly, EditAnywhere, Category = "Description", DisplayName = "Description")
diff --git a/Engine/Plugins/Experimental/PlayerInputDebugger/Source/PlayerInputDebugger/PlayerInputDebugger.Build.cs b/Engine/Plugins/Experimental/PlayerInputDebugger/Source/PlayerInputDebugger/PlayerInputDebugger.Build.cs
--- a/Engine/Plugins/Experimental/PlayerInputDebugger/Source/PlayerInputDebugger/PlayerInputDebugger.Build.cs
+++ b/Engine/Plugins/Experimental/PlayerInputDebugger/Source/PlayerInputDebugger/PlayerInputDebugger.Build.cs
@@ -15,20 +15,22 @@ public class PlayerInputDebugger : ModuleRules
 		PrivateDependencyModuleNames.AddRange(
 			new string[]
 			{
+				"ApplicationCore",
+				"CommonInput",
+				"CommonUI",
 				"Core",
 				"CoreUObject",
 				"Engine",
 				"EnhancedInput",
+				"GameplayTags",
+				"InputCore",
+				"PropertyEditor",
 				"Slate",
 				"SlateCore",
 				"ToolMenus",
-				"UnrealEd",
-				"InputCore",
-				"ApplicationCore",
-				"PropertyEditor",
-				"CommonUI",
-				"CommonInput",
 				"UMG",
+				"UnrealEd",
+				"WorkspaceMenuStructure",
 			});
 	}
 }
diff --git a/Engine/Plugins/Experimental/PlayerInputDebugger/Source/PlayerInputDebugger/Private/PlayerInputDebuggerModule.cpp b/Engine/Plugins/Experimental/PlayerInputDebugger/Source/PlayerInputDebugger/Private/PlayerInputDebuggerModule.cpp
--- a/Engine/Plugins/Experimental/PlayerInputDebugger/Source/PlayerInputDebugger/Private/PlayerInputDebuggerModule.cpp
+++ b/Engine/Plugins/Experimental/PlayerInputDebugger/Source/PlayerInputDebugger/Private/PlayerInputDebuggerModule.cpp
@@ -9,6 +9,8 @@
 #include "Modules/ModuleManager.h"
 #include "ToolMenus.h"
 #include "Widgets/Docking/SDockTab.h"
+#include "WorkspaceMenuStructure.h"
+#include "WorkspaceMenuStructureModule.h"
 
 #define LOCTEXT_NAMESPACE "FPlayerInputDebuggerModule"
 
@@ -37,25 +39,15 @@ void FPlayerInputDebuggerModule::RegisterMenus()
 {
 	FToolMenuOwnerScoped OwnerScoped("PlayerInputDebugger");
 
+	const IWorkspaceMenuStructure& MenuStructure = WorkspaceMenu::GetMenuStructure();
+
 	FGlobalTabmanager::Get()->RegisterNomadTabSpawner(
 		PlayerInputDebuggerTabName,
 		FOnSpawnTab::CreateRaw(this, &FPlayerInputDebuggerModule::SpawnTab))
 		.SetDisplayName(LOCTEXT("TabTitle", "Player Input Debugger"))
-		.SetMenuType(ETabSpawnerMenuType::Hidden);
-
-	if (UToolMenu* ToolsMenu = UToolMenus::Get()->ExtendMenu("MainFrame.MainMenu.Tools"))
-	{
-		FToolMenuSection& Section = ToolsMenu->FindOrAddSection("Debugging", LOCTEXT("DebuggingSection", "Debugging"));
-		Section.AddMenuEntry(
-			"PlayerInputDebugger",
-			LOCTEXT("MenuEntryLabel", "Player Input Debugger"),
-			LOCTEXT("MenuEntryTooltip", "Open the Player Input Debugger window to inspect input events, connected devices, and Enhanced Input state."),
-			FSlateIcon(),
-			FExecuteAction::CreateLambda([]()
-			{
-				FGlobalTabmanager::Get()->TryInvokeTab(FName("PlayerInputDebugger"));
-			}));
-	}
+		.SetTooltipText(LOCTEXT("TabTooltip", "Open the Player Input Debugger window to inspect input events, connected devices, and Enhanced Input state."))
+		.SetGroup(MenuStructure.GetDeveloperToolsDebugCategory())
+		.SetIcon(FSlateIcon("EnhancedInputEditor", "EnhancedInputIcon_Small"));
 }
 
 TSharedRef<SDockTab> FPlayerInputDebuggerModule::SpawnTab(const FSpawnTabArgs& Args)
diff --git a/Engine/Plugins/Experimental/PlayerInputDebugger/Source/PlayerInputDebugger/Private/SEnhancedInputTab.cpp b/Engine/Plugins/Experimental/PlayerInputDebugger/Source/PlayerInputDebugger/Private/SEnhancedInputTab.cpp
--- a/Engine/Plugins/Experimental/PlayerInputDebugger/Source/PlayerInputDebugger/Private/SEnhancedInputTab.cpp
+++ b/Engine/Plugins/Experimental/PlayerInputDebugger/Source/PlayerInputDebugger/Private/SEnhancedInputTab.cpp
@@ -4,11 +4,14 @@
 
 #include "EnhancedActionKeyMapping.h"
 #include "EnhancedInputSubsystems.h"
+#include "GameplayTagContainer.h"
 #include "EnhancedPlayerInput.h"
 #include "InputAction.h"
 #include "InputMappingContext.h"
 #include "UObject/UObjectIterator.h"
+#include "Editor.h"
 #include "Engine/LocalPlayer.h"
+#include "Subsystems/AssetEditorSubsystem.h"
 #include "Framework/Application/SlateApplication.h"
 #include "Styling/AppStyle.h"
 #include "UserSettings/EnhancedInputUserSettings.h"
@@ -16,6 +19,8 @@
 #include "Widgets/Images/SImage.h"
 #include "Widgets/Input/SButton.h"
 #include "Widgets/Input/SEditableTextBox.h"
+#include "Widgets/Input/SHyperlink.h"
+#include "Widgets/Layout/SScrollBox.h"
 #include "Widgets/Layout/SBorder.h"
 #include "Widgets/Layout/SBox.h"
 #include "Widgets/Layout/SExpandableArea.h"
@@ -58,6 +63,38 @@ static FSlateColor GetTriggerEventColor(uint8 Ordinal)
 	}
 }
 
+static const FHyperlinkStyle& GetYellowHyperlinkStyle()
+{
+	static const FHyperlinkStyle Style = []()
+	{
+		FHyperlinkStyle S = FAppStyle::Get().GetWidgetStyle<FHyperlinkStyle>("Hyperlink");
+		S.SetTextStyle(FTextBlockStyle(S.TextStyle)
+			.SetColorAndOpacity(FSlateColor(FLinearColor(0.9f, 0.8f, 0.2f))));
+		return S;
+	}();
+	return Style;
+}
+
+static void OpenAssetEditor(TWeakObjectPtr<const UObject> WeakAsset)
+{
+	const UObject* Asset = WeakAsset.Get();
+	if (!Asset || !GEditor) { return; }
+
+	const FString PathName = Asset->GetPathName();
+
+	UObject* MutableAsset = FindObject<UObject>(nullptr, *PathName);
+	if (!MutableAsset)
+	{
+		MutableAsset = LoadObject<UObject>(nullptr, *PathName);
+	}
+	if (!MutableAsset) { return; }
+
+	if (UAssetEditorSubsystem* Sub = GEditor->GetEditorSubsystem<UAssetEditorSubsystem>())
+	{
+		Sub->OpenEditorForAsset(MutableAsset);
+	}
+}
+
 class SEnhancedInputContextHeaderRow : public SMultiColumnTableRow<TSharedPtr<FEnhancedInputDisplayItem>>
 {
 public:
@@ -88,6 +125,18 @@ public:
 
 		if (ColumnName == TEXT("Context"))
 		{
+			if (Item->WeakAsset.IsValid())
+			{
+				TWeakObjectPtr<const UObject> WeakIMC = Item->WeakAsset;
+				return SNew(SHyperlink)
+					.Text(Item->ContextName)
+					.Style(&GetYellowHyperlinkStyle())
+					.ToolTipText(LOCTEXT("OpenIMCTooltip", "Click to open this Input Mapping Context in the asset editor"))
+					.OnNavigate(FSimpleDelegate::CreateLambda([WeakIMC]()
+					{
+						OpenAssetEditor(WeakIMC);
+					}));
+			}
 			return SNew(STextBlock)
 				.Text(Item->ContextName)
 				.Font(FAppStyle::GetFontStyle("BoldFont"))
@@ -131,10 +180,6 @@ public:
 			LOCTEXT("TECompleted", "Completed"),
 		};
 
-		const FText StateName = TriggerEventNames.IsValidIndex(Item->TriggerEventOrdinal)
-			? TriggerEventNames[Item->TriggerEventOrdinal]
-			: FText::AsNumber(Item->TriggerEventOrdinal);
-
 		// Priority and Context columns are blank for action rows
 		if (ColumnName == TEXT("Priority") || ColumnName == TEXT("Context"))
 		{
@@ -143,6 +188,17 @@ public:
 
 		if (ColumnName == TEXT("Action"))
 		{
+			if (Item->WeakAsset.IsValid())
+			{
+				TWeakObjectPtr<const UObject> WeakAction = Item->WeakAsset;
+				return SNew(SHyperlink)
+					.Text(Item->ActionName)
+					.ToolTipText(LOCTEXT("OpenActionTooltip", "Click to open this Input Action in the asset editor"))
+					.OnNavigate(FSimpleDelegate::CreateLambda([WeakAction]()
+					{
+						OpenAssetEditor(WeakAction);
+					}));
+			}
 			return SNew(STextBlock)
 				.Text(Item->ActionName)
 				.Font(FAppStyle::GetFontStyle("NormalFont"));
@@ -158,16 +214,41 @@ public:
 
 		if (ColumnName == TEXT("State"))
 		{
+			TWeakPtr<FEnhancedInputDisplayItem> WeakItem = Item;
 			return SNew(STextBlock)
-				.Text(StateName)
-				.Font(FAppStyle::GetFontStyle("NormalFont"))
-				.ColorAndOpacity(GetTriggerEventColor(Item->TriggerEventOrdinal));
+				.Text_Lambda([WeakItem]() -> FText
+				{
+					if (TSharedPtr<FEnhancedInputDisplayItem> Pinned = WeakItem.Pin())
+					{
+						return TriggerEventNames.IsValidIndex(Pinned->TriggerEventOrdinal)
+							? TriggerEventNames[Pinned->TriggerEventOrdinal]
+							: FText::AsNumber(Pinned->TriggerEventOrdinal);
+					}
+					return FText::GetEmpty();
+				})
+				.ColorAndOpacity_Lambda([WeakItem]() -> FSlateColor
+				{
+					if (TSharedPtr<FEnhancedInputDisplayItem> Pinned = WeakItem.Pin())
+					{
+						return GetTriggerEventColor(Pinned->TriggerEventOrdinal);
+					}
+					return FSlateColor(FLinearColor(0.8f, 0.8f, 0.8f));
+				})
+				.Font(FAppStyle::GetFontStyle("NormalFont"));
 		}
 
 		if (ColumnName == TEXT("Value"))
 		{
+			TWeakPtr<FEnhancedInputDisplayItem> WeakItem = Item;
 			return SNew(STextBlock)
-				.Text(Item->ValueText)
+				.Text_Lambda([WeakItem]() -> FText
+				{
+					if (TSharedPtr<FEnhancedInputDisplayItem> Pinned = WeakItem.Pin())
+					{
+						return Pinned->ValueText;
+					}
+					return FText::GetEmpty();
+				})
 				.Font(FAppStyle::GetFontStyle("NormalFont"));
 		}
 
@@ -178,6 +259,142 @@ private:
 	TSharedPtr<FEnhancedInputDisplayItem> Item;
 };
 
+// Sub-row displayed directly below each IMC header showing input mode filter settings and registration tracking.
+class SEnhancedInputModeInfoRow : public SMultiColumnTableRow<TSharedPtr<FEnhancedInputDisplayItem>>
+{
+public:
+	SLATE_BEGIN_ARGS(SEnhancedInputModeInfoRow) {}
+		SLATE_ARGUMENT(TSharedPtr<FEnhancedInputDisplayItem>, Item)
+	SLATE_END_ARGS()
+
+	void Construct(const FArguments& InArgs, const TSharedRef<STableViewBase>& InOwnerTableView)
+	{
+		Item = InArgs._Item;
+		SMultiColumnTableRow<TSharedPtr<FEnhancedInputDisplayItem>>::Construct(
+			FSuperRowType::FArguments().Padding(FMargin(2.f, 1.f)),
+			InOwnerTableView
+		);
+	}
+
+	virtual TSharedRef<SWidget> GenerateWidgetForColumn(const FName& ColumnName) override
+	{
+		check(Item.IsValid());
+
+		static const FLinearColor LabelColor(0.5f, 0.5f, 0.5f);
+		static const FLinearColor ValueColor(0.8f, 0.8f, 0.8f);
+		static const FLinearColor PassesColor(0.2f, 0.85f, 0.2f);
+		static const FLinearColor FilteredColor(0.9f, 0.55f, 0.1f);
+
+		const EMappingContextInputModeFilterOptions FilterOptions =
+			static_cast<EMappingContextInputModeFilterOptions>(Item->FilterOptionsOrdinal);
+		const EMappingContextRegistrationTrackingMode RegMode =
+			static_cast<EMappingContextRegistrationTrackingMode>(Item->RegTrackingModeOrdinal);
+
+		// Helper: small dimmed label + normal value side by side
+		auto MakeLabelValue = [&](FText Label, FText Value, FLinearColor ValueCol) -> TSharedRef<SWidget>
+		{
+			return SNew(SHorizontalBox)
+				+ SHorizontalBox::Slot().AutoWidth().VAlign(VAlign_Center).Padding(0.f, 0.f, 4.f, 0.f)
+				[
+					SNew(STextBlock)
+					.Text(Label)
+					.Font(FAppStyle::GetFontStyle("SmallFont"))
+					.ColorAndOpacity(LabelColor)
+				]
+				+ SHorizontalBox::Slot().AutoWidth().VAlign(VAlign_Center)
+				[
+					SNew(STextBlock)
+					.Text(Value)
+					.Font(FAppStyle::GetFontStyle("SmallFont"))
+					.ColorAndOpacity(ValueCol)
+				];
+		};
+
+		if (ColumnName == TEXT("Priority"))
+		{
+			return SNew(STextBlock).Text(FText::GetEmpty());
+		}
+
+		if (ColumnName == TEXT("Context"))
+		{
+			FText FilterModeText;
+			switch (FilterOptions)
+			{
+			case EMappingContextInputModeFilterOptions::UseCustomQuery:
+				FilterModeText = LOCTEXT("FilterCustom",  "Custom Query");   break;
+			case EMappingContextInputModeFilterOptions::DoNotFilter:
+				FilterModeText = LOCTEXT("FilterDisabled","Don't Filter");   break;
+			default:
+				FilterModeText = LOCTEXT("FilterDefault", "Project Default"); break;
+			}
+			return MakeLabelValue(LOCTEXT("FilterLabel", "Filter:"), FilterModeText, ValueColor);
+		}
+
+		if (ColumnName == TEXT("Action"))
+		{
+			return MakeLabelValue(LOCTEXT("QueryLabel", "Query:"), Item->ModeQueryText, ValueColor);
+		}
+
+		if (ColumnName == TEXT("Key"))
+		{
+			return SNew(STextBlock).Text(FText::GetEmpty());
+		}
+
+		if (ColumnName == TEXT("State"))
+		{
+			// Show whether this IMC passes the current EI subsystem input mode filter.
+			// Uses TAttribute lambdas so the widget survives per-frame structural comparisons.
+			TWeakPtr<FEnhancedInputDisplayItem> WeakItem = Item;
+			return SNew(STextBlock)
+				.Text_Lambda([WeakItem]() -> FText
+				{
+					if (TSharedPtr<FEnhancedInputDisplayItem> Pinned = WeakItem.Pin())
+					{
+						const EMappingContextInputModeFilterOptions F =
+							static_cast<EMappingContextInputModeFilterOptions>(Pinned->FilterOptionsOrdinal);
+						if (F == EMappingContextInputModeFilterOptions::DoNotFilter || !Pinned->bShouldApplyModeFilter)
+						{
+							return LOCTEXT("ModeNA", "—");
+						}
+						return Pinned->bPassesModeFilter
+							? LOCTEXT("ModePasses",   "Passes")
+							: LOCTEXT("ModeFiltered", "Filtered");
+					}
+					return FText::GetEmpty();
+				})
+				.ColorAndOpacity_Lambda([WeakItem]() -> FSlateColor
+				{
+					if (TSharedPtr<FEnhancedInputDisplayItem> Pinned = WeakItem.Pin())
+					{
+						const EMappingContextInputModeFilterOptions F =
+							static_cast<EMappingContextInputModeFilterOptions>(Pinned->FilterOptionsOrdinal);
+						if (F == EMappingContextInputModeFilterOptions::DoNotFilter || !Pinned->bShouldApplyModeFilter)
+						{
+							return FSlateColor(LabelColor);
+						}
+						return FSlateColor(Pinned->bPassesModeFilter ? PassesColor : FilteredColor);
+					}
+					return FSlateColor(LabelColor);
+				})
+				.Font(FAppStyle::GetFontStyle("SmallFont"))
+				.ToolTipText(LOCTEXT("ModeFilterTooltip", "Whether this IMC passes the current Enhanced Input subsystem input mode filter"));
+		}
+
+		if (ColumnName == TEXT("Value"))
+		{
+			FText RegText = (RegMode == EMappingContextRegistrationTrackingMode::CountRegistrations)
+				? LOCTEXT("RegCount",     "Count Regs")
+				: LOCTEXT("RegUntracked", "Untracked");
+			return MakeLabelValue(LOCTEXT("RegLabel", "Reg:"), RegText, ValueColor);
+		}
+
+		return SNew(STextBlock).Text(FText::GetEmpty());
+	}
+
+private:
+	TSharedPtr<FEnhancedInputDisplayItem> Item;
+};
+
 void SEnhancedInputTab::Construct(const FArguments& InArgs)
 {
 	SetPlayerController(InArgs._PlayerController.Get());
@@ -367,11 +584,31 @@ void SEnhancedInputTab::Construct(const FArguments& InArgs)
 			.ListItemsSource(&DisplayItems)
 			.OnGenerateRow(this, &SEnhancedInputTab::OnGenerateRow)
 			.HeaderRow(HeaderRow)
-			.SelectionMode(ESelectionMode::None)
+			.SelectionMode(ESelectionMode::Single)
+			.OnSelectionChanged_Lambda([this](TSharedPtr<FEnhancedInputDisplayItem> Item, ESelectInfo::Type)
+			{
+				SelectedItem = Item;
+				RebuildDetails();
+			})
+		]
+
+		+ SVerticalBox::Slot()
+		.AutoHeight()
+		[
+			SNew(SBorder)
+			.BorderImage(FAppStyle::GetBrush("ToolPanel.GroupBorder"))
+			[
+				SNew(SBox)
+				.MaxDesiredHeight(200.f)
+				[
+					SAssignNew(DetailsBox, SBox)
+				]
+			]
 		]
 	];
 
 	RebuildFilterTagsWidget();
+	RebuildDetails();
 }
 
 void SEnhancedInputTab::Tick(const FGeometry& AllottedGeometry, const double InCurrentTime, const float InDeltaTime)
@@ -519,134 +756,387 @@ void SEnhancedInputTab::RebuildFilterTagsWidget()
 
 void SEnhancedInputTab::RebuildDisplayItems()
 {
-	DisplayItems.Empty();
+	// Build candidate items into a local array first.
+	TArray<TSharedPtr<FEnhancedInputDisplayItem>> NewItems;
 
 	APlayerController* PC = WeakPC.Get();
-	if (!PC)
+	ULocalPlayer* LP = PC ? PC->GetLocalPlayer() : nullptr;
+	UEnhancedInputLocalPlayerSubsystem* EIS = LP
+		? LP->GetSubsystem<UEnhancedInputLocalPlayerSubsystem>() : nullptr;
+	UEnhancedPlayerInput* PlayerInput = EIS ? EIS->GetPlayerInput() : nullptr;
+
+	if (PlayerInput)
 	{
-		ListView->RequestListRefresh();
-		return;
+		// Collect and sort contexts by priority (highest first).
+		TArray<TPair<UInputMappingContext*, int32>> SortedContexts;
+		for (TObjectIterator<UInputMappingContext> ContextIt; ContextIt; ++ContextIt)
+		{
+			UInputMappingContext* IMC = *ContextIt;
+			int32 Priority = -1;
+			if (IMC && EIS->HasMappingContext(IMC, Priority))
+			{
+				SortedContexts.Emplace(IMC, Priority);
+			}
+		}
+		SortedContexts.Sort([](const TPair<UInputMappingContext*, int32>& A,
+		                       const TPair<UInputMappingContext*, int32>& B)
+		{
+			return A.Value > B.Value;
+		});
+
+		for (const TPair<UInputMappingContext*, int32>& ContextPair : SortedContexts)
+		{
+			UInputMappingContext* IMC = ContextPair.Key;
+			const int32 Priority = ContextPair.Value;
+
+			// Apply context filter.
+			const bool bHasValidFilter = FilterContexts.ContainsByPredicate(
+				[](const TWeakObjectPtr<UInputMappingContext>& W) { return W.IsValid(); });
+			if (bHasValidFilter)
+			{
+				const bool bPassesFilter = FilterContexts.ContainsByPredicate(
+					[IMC](const TWeakObjectPtr<UInputMappingContext>& W) { return W.Get() == IMC; });
+				if (!bPassesFilter)
+				{
+					continue;
+				}
+			}
+
+			// Context header row
+			TSharedPtr<FEnhancedInputDisplayItem> Header = MakeShared<FEnhancedInputDisplayItem>();
+			Header->bIsHeader = true;
+			Header->ContextName = FText::FromString(IMC->GetName());
+			Header->Priority = Priority;
+			Header->WeakAsset = IMC;
+			NewItems.Add(MoveTemp(Header));
+
+			// Input mode info row
+			{
+				TSharedPtr<FEnhancedInputDisplayItem> InfoRow = MakeShared<FEnhancedInputDisplayItem>();
+				InfoRow->bIsInputModeInfo = true;
+				InfoRow->WeakAsset = IMC;
+
+				const EMappingContextInputModeFilterOptions FilterOptions = IMC->GetInputModeFilterOptions();
+				InfoRow->FilterOptionsOrdinal   = static_cast<uint8>(FilterOptions);
+				InfoRow->RegTrackingModeOrdinal = static_cast<uint8>(IMC->GetRegistrationTrackingMode());
+				InfoRow->bShouldApplyModeFilter = IMC->ShouldFilterMappingByInputMode();
+
+				if (FilterOptions == EMappingContextInputModeFilterOptions::UseCustomQuery)
+				{
+					const FGameplayTagQuery Query = IMC->GetInputModeQuery();
+					InfoRow->ModeQueryText = Query.IsEmpty()
+						? LOCTEXT("ModeQueryEmpty",   "(empty)")
+						: FText::FromString(Query.GetDescription());
+				}
+				else if (FilterOptions == EMappingContextInputModeFilterOptions::UseProjectDefaultQuery)
+				{
+					InfoRow->ModeQueryText = LOCTEXT("ModeQueryDefault", "(project default)");
+				}
+				else
+				{
+					InfoRow->ModeQueryText = LOCTEXT("ModeQueryNone", "—");
+				}
+
+				if (InfoRow->bShouldApplyModeFilter)
+				{
+					InfoRow->bPassesModeFilter = IMC->GetInputModeQuery().Matches(EIS->GetInputMode());
+				}
+
+				NewItems.Add(MoveTemp(InfoRow));
+			}
+
+			// Action mapping rows
+			TArray<TSharedPtr<FEnhancedInputDisplayItem>> ContextRows;
+			for (const FEnhancedActionKeyMapping& Mapping : IMC->GetMappings())
+			{
+				if (!Mapping.Action || Mapping.bShouldBeIgnored)
+				{
+					continue;
+				}
+
+				const FInputActionInstance* Instance = PlayerInput->FindActionInstanceData(Mapping.Action);
+
+				TSharedPtr<FEnhancedInputDisplayItem> Row = MakeShared<FEnhancedInputDisplayItem>();
+				Row->bIsHeader = false;
+				Row->ActionName = FText::FromString(Mapping.Action->GetName());
+				Row->BoundKey = Mapping.Key;
+				Row->WeakAsset = Mapping.Action.Get();
+
+				if (Instance)
+				{
+					Row->TriggerEventOrdinal = static_cast<uint8>(Instance->GetTriggerEvent());
+					Row->ValueText = FText::FromString(Instance->GetValue().ToString());
+				}
+				else
+				{
+					Row->TriggerEventOrdinal = 0;
+					Row->ValueText = LOCTEXT("NoValue", "-");
+				}
+
+				ContextRows.Add(MoveTemp(Row));
+			}
+
+			ContextRows.Sort([](const TSharedPtr<FEnhancedInputDisplayItem>& A,
+			                    const TSharedPtr<FEnhancedInputDisplayItem>& B)
+			{
+				return A->ActionName.CompareTo(B->ActionName) < 0;
+			});
+
+			for (TSharedPtr<FEnhancedInputDisplayItem>& Row : ContextRows)
+			{
+				NewItems.Add(MoveTemp(Row));
+			}
+		}
 	}
 
-	ULocalPlayer* LP = PC->GetLocalPlayer();
-	if (!LP)
+	// Structural comparison: only rebuild the widget tree (RequestListRefresh) when the set of
+	// IMCs, actions, or their structural fields changes. If only live values changed
+	// (trigger state, current value, filter pass status), update the existing items in-place —
+	// the row widgets read those via TAttribute lambdas and will repaint automatically.
+	// This keeps the SHyperlink widgets alive so click events can actually fire.
+	bool bStructureChanged = (NewItems.Num() != DisplayItems.Num());
+	if (!bStructureChanged)
 	{
-		ListView->RequestListRefresh();
-		return;
+		for (int32 i = 0; i < NewItems.Num() && !bStructureChanged; ++i)
+		{
+			const FEnhancedInputDisplayItem& N = *NewItems[i];
+			const FEnhancedInputDisplayItem& D = *DisplayItems[i];
+			bStructureChanged =
+				N.bIsHeader              != D.bIsHeader              ||
+				N.bIsInputModeInfo       != D.bIsInputModeInfo       ||
+				N.Priority               != D.Priority               ||
+				N.BoundKey               != D.BoundKey               ||
+				N.FilterOptionsOrdinal   != D.FilterOptionsOrdinal   ||
+				N.RegTrackingModeOrdinal != D.RegTrackingModeOrdinal ||
+				N.WeakAsset              != D.WeakAsset              ||
+				N.ContextName.ToString() != D.ContextName.ToString() ||
+				N.ActionName.ToString()  != D.ActionName.ToString()  ||
+				N.ModeQueryText.ToString() != D.ModeQueryText.ToString();
+		}
 	}
 
-	UEnhancedInputLocalPlayerSubsystem* EIS = LP->GetSubsystem<UEnhancedInputLocalPlayerSubsystem>();
-	if (!EIS)
+	if (bStructureChanged)
 	{
+		DisplayItems = MoveTemp(NewItems);
 		ListView->RequestListRefresh();
-		return;
 	}
+	else
+	{
+		// Patch live fields in-place; TAttribute lambdas in the row widgets will repaint.
+		for (int32 i = 0; i < NewItems.Num(); ++i)
+		{
+			DisplayItems[i]->TriggerEventOrdinal    = NewItems[i]->TriggerEventOrdinal;
+			DisplayItems[i]->ValueText              = NewItems[i]->ValueText;
+			DisplayItems[i]->bShouldApplyModeFilter = NewItems[i]->bShouldApplyModeFilter;
+			DisplayItems[i]->bPassesModeFilter      = NewItems[i]->bPassesModeFilter;
+		}
+		// No RequestListRefresh — existing row widgets survive, hyperlinks remain clickable.
+	}
+}
 
-	UEnhancedPlayerInput* PlayerInput = EIS->GetPlayerInput();
-	if (!PlayerInput)
+TSharedRef<ITableRow> SEnhancedInputTab::OnGenerateRow(TSharedPtr<FEnhancedInputDisplayItem> Item, const TSharedRef<STableViewBase>& OwnerTable)
+{
+	if (Item->bIsHeader)
 	{
-		ListView->RequestListRefresh();
+		return SNew(SEnhancedInputContextHeaderRow, OwnerTable)
+			.Item(Item);
+	}
+
+	if (Item->bIsInputModeInfo)
+	{
+		return SNew(SEnhancedInputModeInfoRow, OwnerTable)
+			.Item(Item);
+	}
+
+	return SNew(SEnhancedInputRow, OwnerTable)
+		.Item(Item);
+}
+
+void SEnhancedInputTab::RebuildDetails()
+{
+	if (!DetailsBox.IsValid()) { return; }
+
+	if (!SelectedItem.IsValid())
+	{
+		DetailsBox->SetContent(
+			SNew(SBox).Padding(FMargin(8.f))
+			[
+				SNew(STextBlock)
+				.Text(LOCTEXT("NoSelection", "Select a row to view details."))
+				.ColorAndOpacity(FLinearColor(0.5f, 0.5f, 0.5f))
+			]
+		);
 		return;
 	}
 
-	// Collect and sort contexts by priority (highest first).
-	TArray<TPair<const UInputMappingContext*, int32>> SortedContexts;
-	for (TObjectIterator<UInputMappingContext> ContextIt; ContextIt; ++ContextIt)
+	TSharedRef<SScrollBox> ScrollBox = SNew(SScrollBox);
+
+	auto MakeRow = [](FText Label, TSharedRef<SWidget> Value) -> TSharedRef<SWidget>
+	{
+		return SNew(SHorizontalBox)
+			+ SHorizontalBox::Slot().AutoWidth().VAlign(VAlign_Center)
+			[
+				SNew(STextBlock)
+				.Text(Label)
+				.ColorAndOpacity(FLinearColor(0.6f, 0.6f, 0.6f))
+				.MinDesiredWidth(160.f)
+			]
+			+ SHorizontalBox::Slot().AutoWidth().VAlign(VAlign_Center)
+			[
+				Value
+			];
+	};
+
+	auto AddRow = [&](FText Label, TSharedRef<SWidget> Value)
 	{
-		const UInputMappingContext* IMC = *ContextIt;
-		int32 Priority = -1;
-		if (IMC && EIS->HasMappingContext(IMC, Priority))
+		ScrollBox->AddSlot().Padding(FMargin(8.f, 3.f))[ MakeRow(MoveTemp(Label), Value) ];
+	};
+
+	auto AddTextRow = [&](FText Label, FText Value, FLinearColor Color = FLinearColor(0.9f, 0.9f, 0.9f))
+	{
+		AddRow(MoveTemp(Label), SNew(STextBlock).Text(MoveTemp(Value)).ColorAndOpacity(Color));
+	};
+
+	auto MakeAssetWidget = [](FText Name, TWeakObjectPtr<const UObject> WeakAsset) -> TSharedRef<SWidget>
+	{
+		if (WeakAsset.IsValid())
 		{
-			SortedContexts.Emplace(IMC, Priority);
+			return SNew(SHyperlink)
+				.Text(Name)
+				.Style(&GetYellowHyperlinkStyle())
+				.OnNavigate(FSimpleDelegate::CreateLambda([WeakAsset]()
+				{
+					OpenAssetEditor(WeakAsset);
+				}));
 		}
-	}
-	SortedContexts.Sort([](const TPair<const UInputMappingContext*, int32>& A,
-	                       const TPair<const UInputMappingContext*, int32>& B)
-	{
-		return A.Value > B.Value;
-	});
+		return SNew(STextBlock).Text(Name).ColorAndOpacity(FLinearColor(0.9f, 0.9f, 0.9f));
+	};
 
-	for (const TPair<const UInputMappingContext*, int32>& ContextPair : SortedContexts)
+	if (SelectedItem->bIsHeader || SelectedItem->bIsInputModeInfo)
 	{
-		const UInputMappingContext* IMC = ContextPair.Key;
-		const int32 Priority = ContextPair.Value;
-
-		// Apply context filter: only active when at least one slot holds a valid IMC.
-		// Null/unset slots are ignored so partially-filled filter lists don't hide everything.
-		const bool bHasValidFilter = FilterContexts.ContainsByPredicate(
-			[](const TWeakObjectPtr<UInputMappingContext>& W) { return W.IsValid(); });
-		if (bHasValidFilter)
+		const UInputMappingContext* IMC = Cast<const UInputMappingContext>(SelectedItem->WeakAsset.Get());
+
+		// Name + hyperlink
+		const FText DisplayName = IMC
+			? FText::FromString(IMC->GetName())
+			: (SelectedItem->bIsHeader ? SelectedItem->ContextName : LOCTEXT("UnknownIMC", "(unknown)"));
+		AddRow(LOCTEXT("DetailIMCName", "Context:"), MakeAssetWidget(DisplayName, SelectedItem->WeakAsset));
+
+		if (IMC)
 		{
-			const bool bPassesFilter = FilterContexts.ContainsByPredicate(
-				[IMC](const TWeakObjectPtr<UInputMappingContext>& W) { return W.Get() == IMC; });
-			if (!bPassesFilter)
-			{
-				continue;
-			}
+			AddTextRow(LOCTEXT("DetailIMCPath", "Full Path:"),
+				FText::FromString(IMC->GetPathName()),
+				FLinearColor(0.5f, 0.5f, 0.5f));
 		}
 
-		// Context header row
-		TSharedPtr<FEnhancedInputDisplayItem> Header = MakeShared<FEnhancedInputDisplayItem>();
-		Header->bIsHeader = true;
-		Header->ContextName = FText::FromString(IMC->GetName());
-		Header->Priority = Priority;
-		DisplayItems.Add(MoveTemp(Header));
+		if (SelectedItem->bIsHeader)
+		{
+			AddTextRow(LOCTEXT("DetailIMCPriority", "Priority:"), FText::AsNumber(SelectedItem->Priority));
+		}
 
-		// One row per key mapping — collect into a local array so we can sort before committing.
-		TArray<TSharedPtr<FEnhancedInputDisplayItem>> ContextRows;
-		for (const FEnhancedActionKeyMapping& Mapping : IMC->GetMappings())
+		if (IMC)
 		{
-			if (!Mapping.Action || Mapping.bShouldBeIgnored)
+			// Filter mode
+			const EMappingContextInputModeFilterOptions FilterOptions = IMC->GetInputModeFilterOptions();
+			FText FilterText;
+			switch (FilterOptions)
 			{
-				continue;
+			case EMappingContextInputModeFilterOptions::UseCustomQuery:
+				FilterText = LOCTEXT("FilterCustom",   "Custom Query");    break;
+			case EMappingContextInputModeFilterOptions::DoNotFilter:
+				FilterText = LOCTEXT("FilterDisabled", "Don't Filter");    break;
+			default:
+				FilterText = LOCTEXT("FilterDefault",  "Project Default"); break;
 			}
+			AddTextRow(LOCTEXT("DetailFilterMode", "Filter Mode:"), FilterText);
 
-			const FInputActionInstance* Instance = PlayerInput->FindActionInstanceData(Mapping.Action);
-
-			TSharedPtr<FEnhancedInputDisplayItem> Row = MakeShared<FEnhancedInputDisplayItem>();
-			Row->bIsHeader = false;
-			Row->ActionName = FText::FromString(Mapping.Action->GetName());
-			Row->BoundKey = Mapping.Key;
+			// Mode query
+			FText QueryText;
+			if (FilterOptions == EMappingContextInputModeFilterOptions::UseCustomQuery)
+			{
+				const FGameplayTagQuery Query = IMC->GetInputModeQuery();
+				QueryText = Query.IsEmpty()
+					? LOCTEXT("QueryEmpty", "(empty)")
+					: FText::FromString(Query.GetDescription());
+			}
+			else if (FilterOptions == EMappingContextInputModeFilterOptions::UseProjectDefaultQuery)
+			{
+				QueryText = LOCTEXT("QueryDefault", "(project default)");
+			}
+			else
+			{
+				QueryText = LOCTEXT("QueryNA", "—");
+			}
+			AddTextRow(LOCTEXT("DetailModeQuery", "Mode Query:"), QueryText, FLinearColor(0.7f, 0.7f, 0.7f));
 
-			if (Instance)
+			// Passes filter
+			if (IMC->ShouldFilterMappingByInputMode())
 			{
-				Row->TriggerEventOrdinal = static_cast<uint8>(Instance->GetTriggerEvent());
-				Row->ValueText = FText::FromString(Instance->GetValue().ToString());
+				APlayerController* PC = WeakPC.Get();
+				ULocalPlayer* LP = PC ? PC->GetLocalPlayer() : nullptr;
+				UEnhancedInputLocalPlayerSubsystem* EIS = LP
+					? LP->GetSubsystem<UEnhancedInputLocalPlayerSubsystem>() : nullptr;
+				if (EIS)
+				{
+					const bool bPasses = IMC->GetInputModeQuery().Matches(EIS->GetInputMode());
+					AddTextRow(LOCTEXT("DetailPassesFilter", "Passes Filter:"),
+						bPasses ? LOCTEXT("PassesYes", "Yes") : LOCTEXT("PassesNo", "No"),
+						bPasses ? FLinearColor(0.2f, 0.85f, 0.2f) : FLinearColor(0.9f, 0.55f, 0.1f));
+				}
 			}
 			else
 			{
-				Row->TriggerEventOrdinal = 0; // ETriggerEvent::None
-				Row->ValueText = LOCTEXT("NoValue", "-");
+				AddTextRow(LOCTEXT("DetailPassesFilter", "Passes Filter:"),
+					LOCTEXT("PassesNA", "—"), FLinearColor(0.5f, 0.5f, 0.5f));
 			}
 
-			ContextRows.Add(MoveTemp(Row));
+			// Registration tracking
+			const EMappingContextRegistrationTrackingMode RegMode = IMC->GetRegistrationTrackingMode();
+			AddTextRow(LOCTEXT("DetailRegTracking", "Reg Tracking:"),
+				RegMode == EMappingContextRegistrationTrackingMode::CountRegistrations
+					? LOCTEXT("RegCount",     "Count Registrations")
+					: LOCTEXT("RegUntracked", "Untracked"));
 		}
+	}
+	else
+	{
+		// Action row details
+		AddRow(LOCTEXT("DetailActionName", "Action:"),
+			MakeAssetWidget(SelectedItem->ActionName, SelectedItem->WeakAsset));
 
-		// Sort by action name so all mappings for the same action are adjacent.
-		ContextRows.Sort([](const TSharedPtr<FEnhancedInputDisplayItem>& A, const TSharedPtr<FEnhancedInputDisplayItem>& B)
+		if (SelectedItem->WeakAsset.IsValid())
 		{
-			return A->ActionName.CompareTo(B->ActionName) < 0;
-		});
+			AddTextRow(LOCTEXT("DetailActionPath", "Full Path:"),
+				FText::FromString(SelectedItem->WeakAsset->GetPathName()),
+				FLinearColor(0.5f, 0.5f, 0.5f));
+		}
 
-		for (TSharedPtr<FEnhancedInputDisplayItem>& Row : ContextRows)
+		AddTextRow(LOCTEXT("DetailBoundKey", "Bound Key:"),
+			FText::FromString(SelectedItem->BoundKey.ToString()),
+			FLinearColor(0.7f, 0.7f, 0.7f));
+
+		static const TArray<FText> TriggerEventNames =
 		{
-			DisplayItems.Add(MoveTemp(Row));
-		}
-	}
+			LOCTEXT("TENone",      "None"),
+			LOCTEXT("TETriggered", "Triggered"),
+			LOCTEXT("TEStarted",   "Started"),
+			LOCTEXT("TEOngoing",   "Ongoing"),
+			LOCTEXT("TECanceled",  "Canceled"),
+			LOCTEXT("TECompleted", "Completed"),
+		};
+		const FText StateName = TriggerEventNames.IsValidIndex(SelectedItem->TriggerEventOrdinal)
+			? TriggerEventNames[SelectedItem->TriggerEventOrdinal]
+			: FText::AsNumber(SelectedItem->TriggerEventOrdinal);
 
-	ListView->RequestListRefresh();
-}
+		AddRow(LOCTEXT("DetailTriggerState", "Trigger State:"),
+			SNew(STextBlock)
+			.Text(StateName)
+			.ColorAndOpacity(GetTriggerEventColor(SelectedItem->TriggerEventOrdinal)));
 
-TSharedRef<ITableRow> SEnhancedInputTab::OnGenerateRow(TSharedPtr<FEnhancedInputDisplayItem> Item, const TSharedRef<STableViewBase>& OwnerTable)
-{
-	if (Item->bIsHeader)
-	{
-		return SNew(SEnhancedInputContextHeaderRow, OwnerTable)
-			.Item(Item);
+		AddTextRow(LOCTEXT("DetailCurrentValue", "Current Value:"), SelectedItem->ValueText);
 	}
 
-	return SNew(SEnhancedInputRow, OwnerTable)
-		.Item(Item);
+	DetailsBox->SetContent(ScrollBox);
 }
 
 #undef LOCTEXT_NAMESPACE
diff --git a/Engine/Plugins/Experimental/PlayerInputDebugger/Source/PlayerInputDebugger/Private/SEnhancedInputTab.h b/Engine/Plugins/Experimental/PlayerInputDebugger/Source/PlayerInputDebugger/Private/SEnhancedInputTab.h
--- a/Engine/Plugins/Experimental/PlayerInputDebugger/Source/PlayerInputDebugger/Private/SEnhancedInputTab.h
+++ b/Engine/Plugins/Experimental/PlayerInputDebugger/Source/PlayerInputDebugger/Private/SEnhancedInputTab.h
@@ -10,20 +10,32 @@
 class UInputMappingContext;
 class SEditableTextBox;
 
-// A single row in the Enhanced Input display: either a context header or an action mapping row.
+// A single row in the Enhanced Input display: a context header, an IMC mode info row, or an action mapping row.
 struct FEnhancedInputDisplayItem
 {
 	bool bIsHeader = false;
+	bool bIsInputModeInfo = false;
 
 	// Header fields
 	FText ContextName;
 	int32 Priority = 0;
 
+	// Input mode info row fields (appears once per IMC, directly after its header)
+	// Enums stored as uint8 to avoid pulling InputMappingContext.h into this header.
+	uint8 FilterOptionsOrdinal = 0;    // EMappingContextInputModeFilterOptions
+	uint8 RegTrackingModeOrdinal = 0;  // EMappingContextRegistrationTrackingMode
+	FText ModeQueryText;
+	bool bShouldApplyModeFilter = false;  // False when DoNotFilter or global filtering is disabled
+	bool bPassesModeFilter = false;       // Only meaningful when bShouldApplyModeFilter is true
+
 	// Action mapping row fields
 	FText ActionName;
 	FKey BoundKey;
 	uint8 TriggerEventOrdinal = 0;  // ETriggerEvent stored as uint8 to avoid including InputAction.h in header
 	FText ValueText;
+
+	// Asset reference for hyperlink navigation and details panel (IMC for headers/info rows, IA for action rows).
+	TWeakObjectPtr<const UObject> WeakAsset;
 };
 
 // Tab showing active Enhanced Input mapping contexts, their actions, bound keys,
@@ -48,6 +60,7 @@ private:
 	// Context filter helpers
 	void RemoveFilterContext(int32 Index);
 	void RebuildFilterTagsWidget();
+	void RebuildDetails();
 
 	TWeakObjectPtr<APlayerController> WeakPC;
 	TArray<TSharedPtr<FEnhancedInputDisplayItem>> DisplayItems;
@@ -59,4 +72,8 @@ private:
 
 	// Mappable Key Profiles section
 	TSharedPtr<SEditableTextBox> ProfileIdInput;
+
+	// Selection and details panel
+	TSharedPtr<FEnhancedInputDisplayItem> SelectedItem;
+	TSharedPtr<SBox> DetailsBox;
 };
