[Input] Add a native friendly delegate to the UInputDeviceSubsystem for when new input device metadata is available.

Bind to this in the new Player Input Debugger to make it refresh the input devices tab whenever a FInputDeviceId is updated.

#rnx
#Synced-CL 51759702

diff --git a/Engine/Plugins/Experimental/PlayerInputDebugger/Source/PlayerInputDebugger/Private/SInputDevicesTab.cpp b/Engine/Plugins/Experimental/PlayerInputDebugger/Source/PlayerInputDebugger/Private/SInputDevicesTab.cpp
--- a/Engine/Plugins/Experimental/PlayerInputDebugger/Source/PlayerInputDebugger/Private/SInputDevicesTab.cpp
+++ b/Engine/Plugins/Experimental/PlayerInputDebugger/Source/PlayerInputDebugger/Private/SInputDevicesTab.cpp
@@ -93,6 +93,12 @@ void SInputDevicesTab::Construct(const FArguments& InArgs)
 	ConnectionChangedHandle = IPlatformInputDeviceMapper::Get().GetOnInputDeviceConnectionChange().AddRaw(
 		this, &SInputDevicesTab::OnDeviceConnectionChanged);
 
+	if (UInputDeviceSubsystem* DeviceSubsystem = UInputDeviceSubsystem::Get())
+	{
+		HardwareDeviceChangedHandle = DeviceSubsystem->OnInputHardwareDeviceChangedNative.AddRaw(
+			this, &SInputDevicesTab::OnHardwareDeviceChanged);
+	}
+
 	// Seed KnownDevices from currently active devices.
 	{
 		IPlatformInputDeviceMapper& Mapper = IPlatformInputDeviceMapper::Get();
@@ -126,6 +132,11 @@ void SInputDevicesTab::Construct(const FArguments& InArgs)
 SInputDevicesTab::~SInputDevicesTab()
 {
 	IPlatformInputDeviceMapper::Get().GetOnInputDeviceConnectionChange().Remove(ConnectionChangedHandle);
+
+	if (UInputDeviceSubsystem* DeviceSubsystem = UInputDeviceSubsystem::Get())
+	{
+		DeviceSubsystem->OnInputHardwareDeviceChangedNative.Remove(HardwareDeviceChangedHandle);
+	}
 }
 
 void SInputDevicesTab::RebuildDeviceList()
@@ -290,6 +301,18 @@ void SInputDevicesTab::OnDeviceConnectionChanged(EInputDeviceConnectionState New
 	RebuildDeviceList();
 }
 
+void SInputDevicesTab::OnHardwareDeviceChanged(FPlatformUserId UserId, FInputDeviceId DeviceId)
+{
+	if (TSharedPtr<FKnownDevice>* Existing = KnownDevices.Find(DeviceId))
+	{
+		if (UInputDeviceSubsystem* DeviceSubsystem = UInputDeviceSubsystem::Get())
+		{
+			(*Existing)->HardwareId = DeviceSubsystem->GetInputDeviceHardwareIdentifier(DeviceId);
+		}
+		RebuildDeviceList();
+	}
+}
+
 FText SInputDevicesTab::GetConnectionStateText(EInputDeviceConnectionState State) const
 {
 	switch (State)
diff --git a/Engine/Plugins/Experimental/PlayerInputDebugger/Source/PlayerInputDebugger/Private/SInputDevicesTab.h b/Engine/Plugins/Experimental/PlayerInputDebugger/Source/PlayerInputDebugger/Private/SInputDevicesTab.h
--- a/Engine/Plugins/Experimental/PlayerInputDebugger/Source/PlayerInputDebugger/Private/SInputDevicesTab.h
+++ b/Engine/Plugins/Experimental/PlayerInputDebugger/Source/PlayerInputDebugger/Private/SInputDevicesTab.h
@@ -35,12 +35,14 @@ public:
 private:
 	void RebuildDeviceList();
 	void OnDeviceConnectionChanged(EInputDeviceConnectionState NewState, FPlatformUserId UserId, FInputDeviceId DeviceId);
+	void OnHardwareDeviceChanged(FPlatformUserId UserId, FInputDeviceId DeviceId);
 	FText GetConnectionStateText(EInputDeviceConnectionState State) const;
 	FSlateColor GetConnectionStateColor(EInputDeviceConnectionState State) const;
 	FText GetDeviceTypeText(const FHardwareDeviceIdentifier& HardwareId) const;
 
 	TSharedPtr<SVerticalBox> DeviceListContainer;
 	FDelegateHandle ConnectionChangedHandle;
+	FDelegateHandle HardwareDeviceChangedHandle;
 	TMap<FInputDeviceId, TSharedPtr<FKnownDevice>> KnownDevices;
 	TWeakObjectPtr<APlayerController> WeakPC;
 };
diff --git a/Engine/Source/Runtime/Engine/Classes/GameFramework/InputDeviceSubsystem.h b/Engine/Source/Runtime/Engine/Classes/GameFramework/InputDeviceSubsystem.h
--- a/Engine/Source/Runtime/Engine/Classes/GameFramework/InputDeviceSubsystem.h
+++ b/Engine/Source/Runtime/Engine/Classes/GameFramework/InputDeviceSubsystem.h
@@ -111,6 +111,9 @@ struct FActiveDeviceProperty
  */
 DECLARE_DYNAMIC_MULTICAST_DELEGATE_TwoParams(FHardwareInputDeviceChanged, const FPlatformUserId, UserId, const FInputDeviceId, DeviceId);
 
+/** Native delegate for non-UObject listeners for when new input device metadata is available. */
+DECLARE_MULTICAST_DELEGATE_TwoParams(FHardwareInputDeviceChangedNative, FPlatformUserId, FInputDeviceId);
+
 /**
 * The input device subsystem provides an interface to allow users to set Input Device Properties
 * on any Platform User. 
@@ -214,6 +217,9 @@ public:
 	/** A delegate that is fired when a platform user changes what Hardware Input device they are using */
 	UPROPERTY(BlueprintAssignable, Category = "Input Devices")
 	FHardwareInputDeviceChanged OnInputHardwareDeviceChanged;
+
+	/** Native companion to OnInputHardwareDeviceChanged — supports AddRaw/AddWeakLambda for non-UObject listeners. */
+	FHardwareInputDeviceChangedNative OnInputHardwareDeviceChangedNative;
 	
 protected:
 
diff --git a/Engine/Source/Runtime/Engine/Private/GameFramework/InputDeviceSubsystem.cpp b/Engine/Source/Runtime/Engine/Private/GameFramework/InputDeviceSubsystem.cpp
--- a/Engine/Source/Runtime/Engine/Private/GameFramework/InputDeviceSubsystem.cpp
+++ b/Engine/Source/Runtime/Engine/Private/GameFramework/InputDeviceSubsystem.cpp
@@ -611,10 +611,14 @@ void UInputDeviceSubsystem::SetMostRecentlyUsedHardwareDevice(const FInputDevice
 	// Keep track of the latest hardware device and the latest of this type for the given user
 	LatestUserData.LatestDevices.FindOrAdd(InHardwareId.PrimaryDeviceType, InDeviceId);
 	LatestUserData.LatestDeviceId = InDeviceId;
-	
-	if (OnInputHardwareDeviceChanged.IsBound() && bRequiresDelegateUpdate)
+
+	if (bRequiresDelegateUpdate)
 	{
-		OnInputHardwareDeviceChanged.Broadcast(OwningUserId, InDeviceId);
+		OnInputHardwareDeviceChangedNative.Broadcast(OwningUserId, InDeviceId);
+		if (OnInputHardwareDeviceChanged.IsBound())
+		{
+			OnInputHardwareDeviceChanged.Broadcast(OwningUserId, InDeviceId);
+		}	
 	}
 }
 
