[Input debugger] Add missing device types to the player input debugger tab, as well as the additional device name + hardware ID so that it matches up with the "showdebug devices" command.

#rnx
#Synced-CL 51946316

diff --git a/Engine/Plugins/Experimental/PlayerInputDebugger/Source/PlayerInputDebugger/Private/SInputDevicesTab.cpp b/Engine/Plugins/Experimental/PlayerInputDebugger/Source/PlayerInputDebugger/Private/SInputDevicesTab.cpp
--- a/Engine/Plugins/Experimental/PlayerInputDebugger/Source/PlayerInputDebugger/Private/SInputDevicesTab.cpp
+++ b/Engine/Plugins/Experimental/PlayerInputDebugger/Source/PlayerInputDebugger/Private/SInputDevicesTab.cpp
@@ -251,6 +251,19 @@ void SInputDevicesTab::RebuildDeviceList()
 					.ColorAndOpacity(FLinearColor(0.7f, 0.7f, 0.7f))
 				]
 
+				+ SHorizontalBox::Slot()
+				.AutoWidth()
+				.Padding(0.f, 0.f, 8.f, 0.f)
+				.VAlign(VAlign_Center)
+				[
+					SNew(STextBlock)
+					.Text(Device->HardwareId.InputClassName.IsNone()
+						? LOCTEXT("UnknownInputClass", "(unknown class)")
+						: FText::FromName(Device->HardwareId.InputClassName))
+					.Font(FAppStyle::GetFontStyle("NormalFont"))
+					.ColorAndOpacity(FLinearColor(0.5f, 0.5f, 0.5f))
+				]
+
 				+ SHorizontalBox::Slot()
 				.AutoWidth()
 				.Padding(0.f, 0.f, 8.f, 0.f)
@@ -339,10 +352,20 @@ FText SInputDevicesTab::GetDeviceTypeText(const FHardwareDeviceIdentifier& Hardw
 {
 	switch (HardwareId.PrimaryDeviceType)
 	{
-		case EHardwareDevicePrimaryType::KeyboardAndMouse: return LOCTEXT("TypeKBM",     "[Keyboard & Mouse]");
-		case EHardwareDevicePrimaryType::Gamepad:          return LOCTEXT("TypeGamepad", "[Gamepad]");
-		case EHardwareDevicePrimaryType::Touch:            return LOCTEXT("TypeTouch",   "[Touch]");
-		default:                                           return LOCTEXT("TypeUnspec",  "[Unspecified]");
+		case EHardwareDevicePrimaryType::KeyboardAndMouse: return LOCTEXT("TypeKBM",           "[Keyboard & Mouse]");
+		case EHardwareDevicePrimaryType::Gamepad:          return LOCTEXT("TypeGamepad",        "[Gamepad]");
+		case EHardwareDevicePrimaryType::Touch:            return LOCTEXT("TypeTouch",          "[Touch]");
+		case EHardwareDevicePrimaryType::MotionTracking:   return LOCTEXT("TypeMotionTracking", "[Motion Tracking]");
+		case EHardwareDevicePrimaryType::RacingWheel:      return LOCTEXT("TypeRacingWheel",    "[Racing Wheel]");
+		case EHardwareDevicePrimaryType::FlightStick:      return LOCTEXT("TypeFlightStick",    "[Flight Stick]");
+		case EHardwareDevicePrimaryType::Camera:           return LOCTEXT("TypeCamera",         "[Camera]");
+		case EHardwareDevicePrimaryType::Instrument:       return LOCTEXT("TypeInstrument",     "[Instrument]");
+		case EHardwareDevicePrimaryType::CustomTypeA:      return LOCTEXT("TypeCustomA",        "[Custom Type A]");
+		case EHardwareDevicePrimaryType::CustomTypeB:      return LOCTEXT("TypeCustomB",        "[Custom Type B]");
+		case EHardwareDevicePrimaryType::CustomTypeC:      return LOCTEXT("TypeCustomC",        "[Custom Type C]");
+		case EHardwareDevicePrimaryType::CustomTypeD:      return LOCTEXT("TypeCustomD",        "[Custom Type D]");
+		case EHardwareDevicePrimaryType::Unspecified:
+		default:                                           return LOCTEXT("TypeUnspec",         "[Unspecified]");
 	}
 }
 
