[Player Input Debugger] Instead of using buttons as "tabs", actually set up the proper tab manager and window system for the debugger

#rnx
#Synced-CL 52059832

diff --git a/Engine/Plugins/Experimental/PlayerInputDebugger/Source/PlayerInputDebugger/Private/PlayerInputDebuggerModule.cpp b/Engine/Plugins/Experimental/PlayerInputDebugger/Source/PlayerInputDebugger/Private/PlayerInputDebuggerModule.cpp
--- a/Engine/Plugins/Experimental/PlayerInputDebugger/Source/PlayerInputDebugger/Private/PlayerInputDebuggerModule.cpp
+++ b/Engine/Plugins/Experimental/PlayerInputDebugger/Source/PlayerInputDebugger/Private/PlayerInputDebuggerModule.cpp
@@ -52,11 +52,18 @@ void FPlayerInputDebuggerModule::RegisterMenus()
 
 TSharedRef<SDockTab> FPlayerInputDebuggerModule::SpawnTab(const FSpawnTabArgs& Args)
 {
-	return SNew(SDockTab)
-		.TabRole(ETabRole::NomadTab)
-		[
-			SNew(SPlayerInputDebugger)
-		];
+	// The nomad tab must exist before SPlayerInputDebugger is constructed because
+	// it needs to pass itself to FGlobalTabmanager::NewTabManager to create the
+	// local tab manager that drives the dockable sub-tabs.
+	TSharedRef<SDockTab> NomadTab = SNew(SDockTab)
+		.TabRole(ETabRole::NomadTab);
+
+	NomadTab->SetContent(
+		SNew(SPlayerInputDebugger)
+		.ParentTab(NomadTab)
+	);
+
+	return NomadTab;
 }
 
 #undef LOCTEXT_NAMESPACE
diff --git a/Engine/Plugins/Experimental/PlayerInputDebugger/Source/PlayerInputDebugger/Private/SPlayerInputDebugger.cpp b/Engine/Plugins/Experimental/PlayerInputDebugger/Source/PlayerInputDebugger/Private/SPlayerInputDebugger.cpp
--- a/Engine/Plugins/Experimental/PlayerInputDebugger/Source/PlayerInputDebugger/Private/SPlayerInputDebugger.cpp
+++ b/Engine/Plugins/Experimental/PlayerInputDebugger/Source/PlayerInputDebugger/Private/SPlayerInputDebugger.cpp
@@ -10,67 +10,208 @@
 #include "Editor.h"
 #include "Engine/Engine.h"
 #include "Engine/World.h"
+#include "Framework/Commands/UICommandList.h"
+#include "Framework/Docking/LayoutService.h"
+#include "Framework/Docking/TabManager.h"
+#include "Framework/MultiBox/MultiBoxBuilder.h"
 #include "GameFramework/GameModeBase.h"
 #include "GameFramework/PlayerController.h"
 #include "Styling/AppStyle.h"
+#include "Widgets/Docking/SDockTab.h"
 #include "Widgets/Images/SImage.h"
 #include "Widgets/Input/SButton.h"
-#include "Widgets/Input/SCheckBox.h"
 #include "Widgets/Input/SComboBox.h"
 #include "Widgets/Layout/SBorder.h"
-#include "Widgets/Layout/SWidgetSwitcher.h"
 #include "Widgets/Text/STextBlock.h"
 
 #define LOCTEXT_NAMESPACE "SPlayerInputDebugger"
 
+namespace PlayerInputDebuggerTabs
+{
+	static const FName InputEvents   ("PlayerInputDebugger.InputEvents");
+	static const FName InputDevices  ("PlayerInputDebugger.InputDevices");
+	static const FName EnhancedInput ("PlayerInputDebugger.EnhancedInput");
+}
+
 void SPlayerInputDebugger::Construct(const FArguments& InArgs)
 {
-	// Build the three tab content widgets.
+	// Build the three tab content widgets up front so we can pass them into
+	// the spawner lambdas and also call SetPlayerController on PC changes.
 	SAssignNew(PlayerInputEventsTab, SPlayerInputEventsTab);
+	SAssignNew(DevicesTab,           SInputDevicesTab);
+	SAssignNew(EnhancedInputTab,     SEnhancedInputTab);
+	SAssignNew(GlobalInfoSection,    SGlobalInputInfoSection);
 
-	SAssignNew(DevicesTab, SInputDevicesTab);
+	// -----------------------------------------------------------------------
+	// Local tab manager — parented to the host nomad tab so sub-tabs are
+	// scoped to this window and can be dragged/docked within it.
+	// -----------------------------------------------------------------------
+	TSharedPtr<SDockTab> ParentTab = InArgs._ParentTab;
+	check(ParentTab.IsValid());
 
-	SAssignNew(EnhancedInputTab, SEnhancedInputTab);
+	TabManager = FGlobalTabmanager::Get()->NewTabManager(ParentTab.ToSharedRef());
 
-	SAssignNew(GlobalInfoSection, SGlobalInputInfoSection);
+	TabManager->SetOnPersistLayout(FTabManager::FOnPersistLayout::CreateStatic(
+		[](const TSharedRef<FTabManager::FLayout>& InLayout)
+		{
+			if (InLayout->GetPrimaryArea().Pin().IsValid())
+			{
+				FLayoutSaveRestore::SaveToConfig(GEditorLayoutIni, InLayout);
+			}
+		}));
 
-	TSharedRef<SWidgetSwitcher> Switcher = SAssignNew(TabSwitcher, SWidgetSwitcher)
-		+ SWidgetSwitcher::Slot() [ PlayerInputEventsTab.ToSharedRef() ]
-		+ SWidgetSwitcher::Slot() [ DevicesTab.ToSharedRef() ]
-		+ SWidgetSwitcher::Slot() [ EnhancedInputTab.ToSharedRef() ];
+	// Save layout and tear down areas when the host nomad tab is closed.
+	ParentTab->SetOnTabClosed(SDockTab::FOnTabClosedCallback::CreateSP(
+		this, &SPlayerInputDebugger::OnParentTabClosed));
 
-	// Tab bar helper lambda
-	auto MakeTabButton = [this](int32 TabIndex, const FText& Label) -> TSharedRef<SWidget>
+	// Register a spawner for each sub-tab.  The lambdas capture a TSharedRef to
+	// the already-constructed content widget so the same instance is reused if
+	// the tab is closed and re-opened via the Window menu.
+	{
+		TSharedRef<SPlayerInputEventsTab> Content = PlayerInputEventsTab.ToSharedRef();
+		TabManager->RegisterTabSpawner(PlayerInputDebuggerTabs::InputEvents,
+			FOnSpawnTab::CreateLambda([Content](const FSpawnTabArgs&) -> TSharedRef<SDockTab>
+			{
+				return SNew(SDockTab)
+					.TabRole(ETabRole::PanelTab)
+					.Label(LOCTEXT("TabInputEvents", "Input Events"))
+					[ Content ];
+			}))
+			.SetDisplayName(LOCTEXT("TabInputEvents", "Input Events"));
+	}
 	{
-		return SNew(SCheckBox)
-			.Style(FAppStyle::Get(), "ToggleButtonCheckbox")
-			.IsChecked_Lambda([this, TabIndex]()
+		TSharedRef<SInputDevicesTab> Content = DevicesTab.ToSharedRef();
+		TabManager->RegisterTabSpawner(PlayerInputDebuggerTabs::InputDevices,
+			FOnSpawnTab::CreateLambda([Content](const FSpawnTabArgs&) -> TSharedRef<SDockTab>
 			{
-				return ActiveTabIndex == TabIndex ? ECheckBoxState::Checked : ECheckBoxState::Unchecked;
-			})
-			.OnCheckStateChanged_Lambda([this, TabIndex](ECheckBoxState NewState)
+				return SNew(SDockTab)
+					.TabRole(ETabRole::PanelTab)
+					.Label(LOCTEXT("TabInputDevices", "Input Devices"))
+					[ Content ];
+			}))
+			.SetDisplayName(LOCTEXT("TabInputDevices", "Input Devices"));
+	}
+	{
+		TSharedRef<SEnhancedInputTab> Content = EnhancedInputTab.ToSharedRef();
+		TabManager->RegisterTabSpawner(PlayerInputDebuggerTabs::EnhancedInput,
+			FOnSpawnTab::CreateLambda([Content](const FSpawnTabArgs&) -> TSharedRef<SDockTab>
 			{
-				if (NewState == ECheckBoxState::Checked)
-				{
-					ActiveTabIndex = TabIndex;
-					TabSwitcher->SetActiveWidgetIndex(TabIndex);
-				}
-			})
-			.Padding(FMargin(10.f, 4.f))
-			[
-				SNew(STextBlock)
-				.Text(Label)
-			];
+				return SNew(SDockTab)
+					.TabRole(ETabRole::PanelTab)
+					.Label(LOCTEXT("TabEnhancedInput", "Enhanced Input"))
+					[ Content ];
+			}))
+			.SetDisplayName(LOCTEXT("TabEnhancedInput", "Enhanced Input"));
+	}
+
+	// Default layout: all three tabs in a single stack, Enhanced Input foregrounded.
+	// The user can drag tabs out into separate areas; the layout is persisted to
+	// GEditorLayoutIni and restored on next open.
+	TabLayout = FTabManager::NewLayout("Standalone_PlayerInputDebugger_Layout_v1")
+		->AddArea
+		(
+			FTabManager::NewPrimaryArea()
+			->SetOrientation(Orient_Horizontal)
+			->Split
+			(
+				FTabManager::NewStack()
+				->AddTab(PlayerInputDebuggerTabs::InputEvents,   ETabState::OpenedTab)
+				->AddTab(PlayerInputDebuggerTabs::InputDevices,  ETabState::OpenedTab)
+				->AddTab(PlayerInputDebuggerTabs::EnhancedInput, ETabState::OpenedTab)
+				->SetForegroundTab(PlayerInputDebuggerTabs::EnhancedInput)
+			)
+		);
+
+	TabLayout = FLayoutSaveRestore::LoadFromConfig(GEditorLayoutIni, TabLayout.ToSharedRef());
+
+	TSharedRef<SWidget> TabContent = TabManager->RestoreFrom(
+		TabLayout.ToSharedRef(), TSharedPtr<SWindow>()).ToSharedRef();
+
+	// -----------------------------------------------------------------------
+	// Window menu — lets the user re-open sub-tabs they've closed.
+	// -----------------------------------------------------------------------
+	TSharedPtr<FUICommandList> CommandList = MakeShared<FUICommandList>();
+	TWeakPtr<FTabManager> WeakTabManager = TabManager;
+
+	const auto ToggleTab = [](TWeakPtr<FTabManager> InMgr, FName InTab)
+	{
+		if (TSharedPtr<FTabManager> Mgr = InMgr.Pin())
+		{
+			if (TSharedPtr<SDockTab> Existing = Mgr->FindExistingLiveTab(InTab))
+			{
+				Existing->RequestCloseTab();
+			}
+			else
+			{
+				Mgr->TryInvokeTab(InTab);
+			}
+		}
+	};
+	const auto IsTabVisible = [](TWeakPtr<FTabManager> InMgr, FName InTab) -> bool
+	{
+		if (TSharedPtr<FTabManager> Mgr = InMgr.Pin())
+		{
+			return Mgr->FindExistingLiveTab(InTab).IsValid();
+		}
+		return false;
 	};
 
+	FMenuBarBuilder MenuBarBuilder(CommandList);
+	MenuBarBuilder.AddPullDownMenu(
+		LOCTEXT("WindowMenuLabel", "Window"),
+		FText::GetEmpty(),
+		FNewMenuDelegate::CreateLambda([WeakTabManager, ToggleTab, IsTabVisible](FMenuBuilder& Builder)
+		{
+			Builder.AddMenuEntry(
+				LOCTEXT("TabInputEvents",  "Input Events"),
+				FText::GetEmpty(), FSlateIcon(),
+				FUIAction(
+					FExecuteAction::CreateLambda([WeakTabManager, ToggleTab]()  { ToggleTab(WeakTabManager, PlayerInputDebuggerTabs::InputEvents); }),
+					FCanExecuteAction(),
+					FIsActionChecked::CreateLambda([WeakTabManager, IsTabVisible]() { return IsTabVisible(WeakTabManager, PlayerInputDebuggerTabs::InputEvents); })
+				), NAME_None, EUserInterfaceActionType::ToggleButton);
+
+			Builder.AddMenuEntry(
+				LOCTEXT("TabInputDevices", "Input Devices"),
+				FText::GetEmpty(), FSlateIcon(),
+				FUIAction(
+					FExecuteAction::CreateLambda([WeakTabManager, ToggleTab]()  { ToggleTab(WeakTabManager, PlayerInputDebuggerTabs::InputDevices); }),
+					FCanExecuteAction(),
+					FIsActionChecked::CreateLambda([WeakTabManager, IsTabVisible]() { return IsTabVisible(WeakTabManager, PlayerInputDebuggerTabs::InputDevices); })
+				), NAME_None, EUserInterfaceActionType::ToggleButton);
+
+			Builder.AddMenuEntry(
+				LOCTEXT("TabEnhancedInput", "Enhanced Input"),
+				FText::GetEmpty(), FSlateIcon(),
+				FUIAction(
+					FExecuteAction::CreateLambda([WeakTabManager, ToggleTab]()  { ToggleTab(WeakTabManager, PlayerInputDebuggerTabs::EnhancedInput); }),
+					FCanExecuteAction(),
+					FIsActionChecked::CreateLambda([WeakTabManager, IsTabVisible]() { return IsTabVisible(WeakTabManager, PlayerInputDebuggerTabs::EnhancedInput); })
+				), NAME_None, EUserInterfaceActionType::ToggleButton);
+		})
+	);
+
+	TSharedRef<SWidget> MenuBarWidget = MenuBarBuilder.MakeWidget();
+	TabManager->SetMenuMultiBox(MenuBarBuilder.GetMultiBox(), MenuBarWidget);
+
+	// -----------------------------------------------------------------------
+	// Root layout: menu bar / PC selector / global info / tab area
+	// -----------------------------------------------------------------------
 	ChildSlot
 	[
 		SNew(SVerticalBox)
 
+		// Window menu bar
+		+ SVerticalBox::Slot()
+		.AutoHeight()
+		[
+			MenuBarWidget
+		]
+
 		// Player controller selector
 		+ SVerticalBox::Slot()
 		.AutoHeight()
-		.Padding(6.f, 6.f, 6.f, 4.f)
+		.Padding(6.f, 4.f, 6.f, 4.f)
 		[
 			SNew(SHorizontalBox)
 
@@ -119,48 +260,23 @@ void SPlayerInputDebugger::Construct(const FArguments& InArgs)
 			]
 		]
 
-		// Tab bar
-		+ SVerticalBox::Slot()
-		.AutoHeight()
-		[
-			SNew(SBorder)
-			.BorderImage(FAppStyle::GetBrush("ToolPanel.GroupBorder"))
-			.Padding(FMargin(0.f))
-			[
-				SNew(SHorizontalBox)
-
-				+ SHorizontalBox::Slot()
-				.AutoWidth()
-				[
-					MakeTabButton(0, LOCTEXT("TabInputEvents", "Input Events"))
-				]
-
-				+ SHorizontalBox::Slot()
-				.AutoWidth()
-				[
-					MakeTabButton(1, LOCTEXT("TabInputDevices", "Input Devices"))
-				]
-
-				+ SHorizontalBox::Slot()
-				.AutoWidth()
-				[
-					MakeTabButton(2, LOCTEXT("TabEnhancedInput", "Enhanced Input"))
-				]
-			]
-		]
-
-		// Global input overview (always visible, below tab bar)
+		// Global input overview — always visible above the tab area
 		+ SVerticalBox::Slot()
 		.AutoHeight()
 		[
 			GlobalInfoSection.ToSharedRef()
 		]
 
-		// Tab content
+		// Dockable tab area
 		+ SVerticalBox::Slot()
 		.FillHeight(1.f)
 		[
-			Switcher
+			SNew(SBorder)
+			.BorderImage(FAppStyle::GetBrush("ToolPanel.DarkGroupBorder"))
+			.Padding(FMargin(0.f, 2.f))
+			[
+				TabContent
+			]
 		]
 	];
 
@@ -176,12 +292,27 @@ void SPlayerInputDebugger::Construct(const FArguments& InArgs)
 
 SPlayerInputDebugger::~SPlayerInputDebugger()
 {
+	if (TabManager.IsValid())
+	{
+		TabManager->UnregisterAllTabSpawners();
+		TabManager.Reset();
+	}
+
 	FEditorDelegates::BeginPIE.Remove(BeginPIEHandle);
 	FEditorDelegates::EndPIE.Remove(EndPIEHandle);
 	FGameModeEvents::OnGameModePostLoginEvent().Remove(PostLoginHandle);
 	FGameModeEvents::OnGameModeLogoutEvent().Remove(LogoutHandle);
 }
 
+void SPlayerInputDebugger::OnParentTabClosed(TSharedRef<SDockTab> /*ClosedTab*/)
+{
+	if (TabManager.IsValid())
+	{
+		FLayoutSaveRestore::SaveToConfig(GEditorLayoutIni, TabManager->PersistLayout());
+		TabManager->CloseAllAreas();
+	}
+}
+
 void SPlayerInputDebugger::RefreshPlayerControllers(const bool bIsPlayEnding /*= false*/)
 {
 	PCEntries.Empty();
diff --git a/Engine/Plugins/Experimental/PlayerInputDebugger/Source/PlayerInputDebugger/Private/SPlayerInputDebugger.h b/Engine/Plugins/Experimental/PlayerInputDebugger/Source/PlayerInputDebugger/Private/SPlayerInputDebugger.h
--- a/Engine/Plugins/Experimental/PlayerInputDebugger/Source/PlayerInputDebugger/Private/SPlayerInputDebugger.h
+++ b/Engine/Plugins/Experimental/PlayerInputDebugger/Source/PlayerInputDebugger/Private/SPlayerInputDebugger.h
@@ -4,11 +4,13 @@
 
 #include "GameFramework/PlayerController.h"
 #include "Widgets/Input/SComboBox.h"
-#include "Widgets/Layout/SWidgetSwitcher.h"
 #include "Widgets/SCompoundWidget.h"
+#include "Framework/Docking/TabManager.h"
 
 class AController;
 class AGameModeBase;
+class FTabManager;
+class SDockTab;
 class SInputDevicesTab;
 class SEnhancedInputTab;
 class SPlayerInputEventsTab;
@@ -22,11 +24,14 @@ struct FPlayerControllerEntry
 };
 
 // Root widget for the Player Input Debugger window.
-// Contains a player controller selector and three sub-tabs.
+// Contains a player controller selector, a persistent global info section,
+// and three dockable sub-tabs managed by a local FTabManager.
 class SPlayerInputDebugger : public SCompoundWidget
 {
 public:
 	SLATE_BEGIN_ARGS(SPlayerInputDebugger) {}
+		// The nomad SDockTab that hosts this widget. Required to create the local tab manager.
+		SLATE_ARGUMENT(TSharedPtr<SDockTab>, ParentTab)
 	SLATE_END_ARGS()
 
 	void Construct(const FArguments& InArgs);
@@ -42,19 +47,21 @@ private:
 	TSharedRef<SWidget> MakePCComboEntry(TSharedPtr<FPlayerControllerEntry> Entry);
 	FText GetSelectedPCText() const;
 	void SelectEntry(TSharedPtr<FPlayerControllerEntry> Entry);
+	void OnParentTabClosed(TSharedRef<SDockTab> ClosedTab);
 
 	TArray<TSharedPtr<FPlayerControllerEntry>> PCEntries;
 	TSharedPtr<FPlayerControllerEntry> SelectedEntry;
 
 	TSharedPtr<SComboBox<TSharedPtr<FPlayerControllerEntry>>> PCCombo;
-	TSharedPtr<SWidgetSwitcher> TabSwitcher;
-	int32 ActiveTabIndex = 0;
 
 	TSharedPtr<SPlayerInputEventsTab> PlayerInputEventsTab;
 	TSharedPtr<SInputDevicesTab> DevicesTab;
 	TSharedPtr<SEnhancedInputTab> EnhancedInputTab;
 	TSharedPtr<SGlobalInputInfoSection> GlobalInfoSection;
 
+	TSharedPtr<FTabManager> TabManager;
+	TSharedPtr<FTabManager::FLayout> TabLayout;
+
 	FDelegateHandle BeginPIEHandle;
 	FDelegateHandle EndPIEHandle;
 	FDelegateHandle PostLoginHandle;
