[Player input debugger] Add a delegate to the input devices tab to refresh when the device pairing changes. Let you filter the enhanced input tab by input actions as well as IMCs, and by FKey type (gamepad, keyboard, touch)

#Synced-CL 52059832

diff --git a/Engine/Plugins/Experimental/PlayerInputDebugger/Source/PlayerInputDebugger/Private/SEnhancedInputTab.cpp b/Engine/Plugins/Experimental/PlayerInputDebugger/Source/PlayerInputDebugger/Private/SEnhancedInputTab.cpp
--- a/Engine/Plugins/Experimental/PlayerInputDebugger/Source/PlayerInputDebugger/Private/SEnhancedInputTab.cpp
+++ b/Engine/Plugins/Experimental/PlayerInputDebugger/Source/PlayerInputDebugger/Private/SEnhancedInputTab.cpp
@@ -18,6 +18,7 @@
 #include "PropertyCustomizationHelpers.h"
 #include "Widgets/Images/SImage.h"
 #include "Widgets/Input/SButton.h"
+#include "Widgets/Input/SCheckBox.h"
 #include "Widgets/Input/SEditableTextBox.h"
 #include "Widgets/Input/SHyperlink.h"
 #include "Widgets/Layout/SScrollBox.h"
@@ -417,6 +418,7 @@ void SEnhancedInputTab::Construct(const FArguments& InArgs)
 		+ SHeaderRow::Column("Key")
 		.DefaultLabel(LOCTEXT("ColKey", "Key"))
 		.ManualWidth(EnhancedInputTab::KeyColumnWidth)
+		.OnGetMenuContent(FOnGetContent::CreateSP(this, &SEnhancedInputTab::MakeKeyFilterMenuContent))
 
 		+ SHeaderRow::Column("State")
 		.DefaultLabel(LOCTEXT("ColState", "Trigger State"))
@@ -452,6 +454,19 @@ void SEnhancedInputTab::Construct(const FArguments& InArgs)
 			]
 		]
 
+		+ SVerticalBox::Slot()
+		.AutoHeight()
+		[
+			SNew(SExpandableArea)
+			.AreaTitle(LOCTEXT("ActionFiltersTitle", "Action Filters"))
+			.InitiallyCollapsed(false)
+			.BodyContent()
+			[
+				SAssignNew(ActionFilterTagsContainer, SBox)
+				.Padding(FMargin(4.f, 4.f))
+			]
+		]
+
 		+ SVerticalBox::Slot()
 		.AutoHeight()
 		[
@@ -580,16 +595,46 @@ void SEnhancedInputTab::Construct(const FArguments& InArgs)
 		+ SVerticalBox::Slot()
 		.FillHeight(1.f)
 		[
-			SAssignNew(ListView, SListView<TSharedPtr<FEnhancedInputDisplayItem>>)
-			.ListItemsSource(&DisplayItems)
-			.OnGenerateRow(this, &SEnhancedInputTab::OnGenerateRow)
-			.HeaderRow(HeaderRow)
-			.SelectionMode(ESelectionMode::Single)
-			.OnSelectionChanged_Lambda([this](TSharedPtr<FEnhancedInputDisplayItem> Item, ESelectInfo::Type)
-			{
-				SelectedItem = Item;
-				RebuildDetails();
-			})
+			SNew(SOverlay)
+
+			+ SOverlay::Slot()
+			[
+				SAssignNew(ListView, SListView<TSharedPtr<FEnhancedInputDisplayItem>>)
+				.ListItemsSource(&DisplayItems)
+				.OnGenerateRow(this, &SEnhancedInputTab::OnGenerateRow)
+				.HeaderRow(HeaderRow)
+				.SelectionMode(ESelectionMode::Single)
+				.OnSelectionChanged_Lambda([this](TSharedPtr<FEnhancedInputDisplayItem> Item, ESelectInfo::Type)
+				{
+					SelectedItem = Item;
+					RebuildDetails();
+				})
+			]
+
+			// Shown when all rows are hidden by the active filters.
+			+ SOverlay::Slot()
+			.HAlign(HAlign_Center)
+			.VAlign(VAlign_Center)
+			[
+				SNew(STextBlock)
+				.Visibility_Lambda([this]() -> EVisibility
+				{
+					if (!DisplayItems.IsEmpty())
+					{
+						return EVisibility::Collapsed;
+					}
+					const bool bAnyContextFilter = FilterContexts.ContainsByPredicate(
+						[](const TWeakObjectPtr<UInputMappingContext>& W) { return W.IsValid(); });
+					const bool bAnyActionFilter = FilterActions.ContainsByPredicate(
+						[](const TWeakObjectPtr<UInputAction>& W) { return W.IsValid(); });
+					const bool bAnyKeyTypeFilter = !bShowKeyboardKeys || !bShowGamepadKeys || !bShowTouchKeys;
+					return (bAnyContextFilter || bAnyActionFilter || bAnyKeyTypeFilter)
+						? EVisibility::HitTestInvisible
+						: EVisibility::Collapsed;
+				})
+				.Text(LOCTEXT("AllFilteredOut", "All mappings are hidden by the active filters."))
+				.ColorAndOpacity(FLinearColor(0.5f, 0.5f, 0.5f))
+			]
 		]
 
 		+ SVerticalBox::Slot()
@@ -608,6 +653,7 @@ void SEnhancedInputTab::Construct(const FArguments& InArgs)
 	];
 
 	RebuildFilterTagsWidget();
+	RebuildActionFilterTagsWidget();
 	RebuildDetails();
 }
 
@@ -626,6 +672,83 @@ void SEnhancedInputTab::RemoveFilterContext(int32 Index)
 	}
 }
 
+void SEnhancedInputTab::RemoveFilterAction(int32 Index)
+{
+	if (FilterActions.IsValidIndex(Index))
+	{
+		FilterActions.RemoveAt(Index);
+		RebuildActionFilterTagsWidget();
+	}
+}
+
+TSharedRef<SWidget> SEnhancedInputTab::MakeKeyFilterMenuContent()
+{
+	// Called by the Key column header caret to produce the filter dropdown.
+	// Uses IsChecked_Lambda so each open reads the current state from this widget.
+	return SNew(SBox)
+		.Padding(FMargin(4.f, 2.f))
+		[
+			SNew(SVerticalBox)
+
+			+ SVerticalBox::Slot()
+			.AutoHeight()
+			.Padding(2.f)
+			[
+				SNew(SCheckBox)
+				.IsChecked_Lambda([this]()
+				{
+					return bShowKeyboardKeys ? ECheckBoxState::Checked : ECheckBoxState::Unchecked;
+				})
+				.OnCheckStateChanged_Lambda([this](ECheckBoxState State)
+				{
+					bShowKeyboardKeys = (State == ECheckBoxState::Checked);
+				})
+				[
+					SNew(STextBlock)
+					.Text(LOCTEXT("ShowKeyboard", "Keyboard"))
+				]
+			]
+
+			+ SVerticalBox::Slot()
+			.AutoHeight()
+			.Padding(2.f)
+			[
+				SNew(SCheckBox)
+				.IsChecked_Lambda([this]()
+				{
+					return bShowGamepadKeys ? ECheckBoxState::Checked : ECheckBoxState::Unchecked;
+				})
+				.OnCheckStateChanged_Lambda([this](ECheckBoxState State)
+				{
+					bShowGamepadKeys = (State == ECheckBoxState::Checked);
+				})
+				[
+					SNew(STextBlock)
+					.Text(LOCTEXT("ShowGamepad", "Gamepad"))
+				]
+			]
+
+			+ SVerticalBox::Slot()
+			.AutoHeight()
+			.Padding(2.f)
+			[
+				SNew(SCheckBox)
+				.IsChecked_Lambda([this]()
+				{
+					return bShowTouchKeys ? ECheckBoxState::Checked : ECheckBoxState::Unchecked;
+				})
+				.OnCheckStateChanged_Lambda([this](ECheckBoxState State)
+				{
+					bShowTouchKeys = (State == ECheckBoxState::Checked);
+				})
+				[
+					SNew(STextBlock)
+					.Text(LOCTEXT("ShowTouch", "Touch"))
+				]
+			]
+		];
+}
+
 void SEnhancedInputTab::RebuildFilterTagsWidget()
 {
 	TSharedRef<SVerticalBox> Content = SNew(SVerticalBox);
@@ -754,6 +877,133 @@ void SEnhancedInputTab::RebuildFilterTagsWidget()
 	FilterTagsContainer->SetContent(Content);
 }
 
+void SEnhancedInputTab::RebuildActionFilterTagsWidget()
+{
+	TSharedRef<SVerticalBox> Content = SNew(SVerticalBox);
+
+	const int32 ValidCount = FilterActions.FilterByPredicate(
+		[](const TWeakObjectPtr<UInputAction>& W) { return W.IsValid(); }).Num();
+
+	// Header row: status hint + "+" add-slot button
+	Content->AddSlot()
+	.AutoHeight()
+	.Padding(0.f, 0.f, 0.f, 4.f)
+	[
+		SNew(SHorizontalBox)
+
+		+ SHorizontalBox::Slot()
+		.FillWidth(1.f)
+		.VAlign(VAlign_Center)
+		[
+			SNew(STextBlock)
+			.Text(ValidCount == 0
+				? LOCTEXT("ActionFilterNone", "No filter active — showing all actions")
+				: FText::Format(LOCTEXT("ActionFilterActive", "Filtering to {0} action(s)"), FText::AsNumber(ValidCount)))
+			.ColorAndOpacity(ValidCount == 0
+				? FSlateColor(FLinearColor(0.5f, 0.5f, 0.5f))
+				: FSlateColor(FSlateColor::UseForeground()))
+		]
+
+		+ SHorizontalBox::Slot()
+		.AutoWidth()
+		.VAlign(VAlign_Center)
+		[
+			SNew(SButton)
+			.ButtonStyle(FAppStyle::Get(), "SimpleButton")
+			.ToolTipText(LOCTEXT("AddActionFilterTooltip", "Add an action filter slot"))
+			.OnClicked_Lambda([this]() -> FReply
+			{
+				FilterActions.Add(nullptr);
+				RebuildActionFilterTagsWidget();
+				return FReply::Handled();
+			})
+			[
+				SNew(SImage)
+				.Image(FAppStyle::GetBrush("Icons.Plus"))
+				.DesiredSizeOverride(FVector2D(16.f, 16.f))
+			]
+		]
+	];
+
+	// One asset-picker row per filter slot
+	for (int32 i = 0; i < FilterActions.Num(); ++i)
+	{
+		const int32 CapturedIndex = i;
+
+		Content->AddSlot()
+		.AutoHeight()
+		.Padding(0.f, 1.f)
+		[
+			SNew(SHorizontalBox)
+
+			// Array index badge
+			+ SHorizontalBox::Slot()
+			.AutoWidth()
+			.VAlign(VAlign_Center)
+			.Padding(0.f, 0.f, 6.f, 0.f)
+			[
+				SNew(STextBlock)
+				.Text(FText::AsNumber(i))
+				.ColorAndOpacity(FLinearColor(0.5f, 0.5f, 0.5f))
+				.MinDesiredWidth(18.f)
+			]
+
+			// UE property-style asset picker
+			+ SHorizontalBox::Slot()
+			.FillWidth(1.f)
+			[
+				SNew(SObjectPropertyEntryBox)
+				.AllowedClass(UInputAction::StaticClass())
+				.ObjectPath(TAttribute<FString>::CreateLambda([this, CapturedIndex]() -> FString
+				{
+					if (FilterActions.IsValidIndex(CapturedIndex))
+					{
+						if (const UInputAction* IA = FilterActions[CapturedIndex].Get())
+						{
+							return IA->GetPathName();
+						}
+					}
+					return FString();
+				}))
+				.OnObjectChanged(FOnSetObject::CreateLambda([this, CapturedIndex](const FAssetData& AssetData)
+				{
+					if (FilterActions.IsValidIndex(CapturedIndex))
+					{
+						FilterActions[CapturedIndex] = Cast<UInputAction>(AssetData.GetAsset());
+						RebuildActionFilterTagsWidget();
+					}
+				}))
+				.AllowClear(true)
+				.DisplayUseSelected(true)
+				.DisplayBrowse(true)
+			]
+
+			// Delete-slot button
+			+ SHorizontalBox::Slot()
+			.AutoWidth()
+			.VAlign(VAlign_Center)
+			.Padding(4.f, 0.f, 0.f, 0.f)
+			[
+				SNew(SButton)
+				.ButtonStyle(FAppStyle::Get(), "SimpleButton")
+				.ToolTipText(LOCTEXT("RemoveActionFilterTooltip", "Remove this filter slot"))
+				.OnClicked_Lambda([this, CapturedIndex]() -> FReply
+				{
+					RemoveFilterAction(CapturedIndex);
+					return FReply::Handled();
+				})
+				[
+					SNew(SImage)
+					.Image(FAppStyle::GetBrush("Icons.Delete"))
+					.DesiredSizeOverride(FVector2D(16.f, 16.f))
+				]
+			]
+		];
+	}
+
+	ActionFilterTagsContainer->SetContent(Content);
+}
+
 void SEnhancedInputTab::RebuildDisplayItems()
 {
 	// Build candidate items into a local array first.
@@ -854,6 +1104,30 @@ void SEnhancedInputTab::RebuildDisplayItems()
 					continue;
 				}
 
+				// Apply action filter: if any valid action is in the filter list, only show mappings for those actions.
+				{
+					const bool bHasValidActionFilter = FilterActions.ContainsByPredicate(
+						[](const TWeakObjectPtr<UInputAction>& W) { return W.IsValid(); });
+					if (bHasValidActionFilter && !FilterActions.ContainsByPredicate(
+						[&Mapping](const TWeakObjectPtr<UInputAction>& W) { return W.Get() == Mapping.Action.Get(); }))
+					{
+						continue;
+					}
+				}
+
+				// Apply key type filter.
+				{
+					const bool bIsGamepadKey  = Mapping.Key.IsGamepadKey();
+					const bool bIsTouchKey    = Mapping.Key.IsTouch();
+					const bool bIsKeyboardKey = !bIsGamepadKey && !bIsTouchKey;
+					if ((bIsGamepadKey  && !bShowGamepadKeys)  ||
+						(bIsTouchKey    && !bShowTouchKeys)    ||
+						(bIsKeyboardKey && !bShowKeyboardKeys))
+					{
+						continue;
+					}
+				}
+
 				const FInputActionInstance* Instance = PlayerInput->FindActionInstanceData(Mapping.Action);
 
 				TSharedPtr<FEnhancedInputDisplayItem> Row = MakeShared<FEnhancedInputDisplayItem>();
diff --git a/Engine/Plugins/Experimental/PlayerInputDebugger/Source/PlayerInputDebugger/Private/SEnhancedInputTab.h b/Engine/Plugins/Experimental/PlayerInputDebugger/Source/PlayerInputDebugger/Private/SEnhancedInputTab.h
--- a/Engine/Plugins/Experimental/PlayerInputDebugger/Source/PlayerInputDebugger/Private/SEnhancedInputTab.h
+++ b/Engine/Plugins/Experimental/PlayerInputDebugger/Source/PlayerInputDebugger/Private/SEnhancedInputTab.h
@@ -7,6 +7,7 @@
 #include "Widgets/Views/ITableRow.h"
 #include "Widgets/Views/SListView.h"
 
+class UInputAction;
 class UInputMappingContext;
 class SEditableTextBox;
 
@@ -60,6 +61,14 @@ private:
 	// Context filter helpers
 	void RemoveFilterContext(int32 Index);
 	void RebuildFilterTagsWidget();
+
+	// Action filter helpers
+	void RemoveFilterAction(int32 Index);
+	void RebuildActionFilterTagsWidget();
+
+	// Key type filter — dropdown menu content for the Key column header
+	TSharedRef<SWidget> MakeKeyFilterMenuContent();
+
 	void RebuildDetails();
 
 	TWeakObjectPtr<APlayerController> WeakPC;
@@ -70,6 +79,15 @@ private:
 	TArray<TWeakObjectPtr<UInputMappingContext>> FilterContexts;
 	TSharedPtr<SBox> FilterTagsContainer;
 
+	// Actions to show; empty = show all.
+	TArray<TWeakObjectPtr<UInputAction>> FilterActions;
+	TSharedPtr<SBox> ActionFilterTagsContainer;
+
+	// Key type visibility flags (default all visible)
+	bool bShowKeyboardKeys = true;
+	bool bShowGamepadKeys  = true;
+	bool bShowTouchKeys    = true;
+
 	// Mappable Key Profiles section
 	TSharedPtr<SEditableTextBox> ProfileIdInput;
 
diff --git a/Engine/Plugins/Experimental/PlayerInputDebugger/Source/PlayerInputDebugger/Private/SInputDevicesTab.cpp b/Engine/Plugins/Experimental/PlayerInputDebugger/Source/PlayerInputDebugger/Private/SInputDevicesTab.cpp
--- a/Engine/Plugins/Experimental/PlayerInputDebugger/Source/PlayerInputDebugger/Private/SInputDevicesTab.cpp
+++ b/Engine/Plugins/Experimental/PlayerInputDebugger/Source/PlayerInputDebugger/Private/SInputDevicesTab.cpp
@@ -93,6 +93,9 @@ void SInputDevicesTab::Construct(const FArguments& InArgs)
 	ConnectionChangedHandle = IPlatformInputDeviceMapper::Get().GetOnInputDeviceConnectionChange().AddRaw(
 		this, &SInputDevicesTab::OnDeviceConnectionChanged);
 
+	PairingChangeHandle = IPlatformInputDeviceMapper::Get().GetOnInputDevicePairingChange().AddRaw(
+		this, &SInputDevicesTab::OnDevicePairingChanged);
+
 	if (UInputDeviceSubsystem* DeviceSubsystem = UInputDeviceSubsystem::Get())
 	{
 		HardwareDeviceChangedHandle = DeviceSubsystem->OnInputHardwareDeviceChangedNative.AddRaw(
@@ -132,6 +135,7 @@ void SInputDevicesTab::Construct(const FArguments& InArgs)
 SInputDevicesTab::~SInputDevicesTab()
 {
 	IPlatformInputDeviceMapper::Get().GetOnInputDeviceConnectionChange().Remove(ConnectionChangedHandle);
+	IPlatformInputDeviceMapper::Get().GetOnInputDevicePairingChange().Remove(PairingChangeHandle);
 
 	if (UInputDeviceSubsystem* DeviceSubsystem = UInputDeviceSubsystem::Get())
 	{
@@ -326,6 +330,29 @@ void SInputDevicesTab::OnHardwareDeviceChanged(FPlatformUserId UserId, FInputDev
 	}
 }
 
+void SInputDevicesTab::OnDevicePairingChanged(FInputDeviceId DeviceId, FPlatformUserId NewUserId, FPlatformUserId OldUserId)
+{
+	if (TSharedPtr<FKnownDevice>* Existing = KnownDevices.Find(DeviceId))
+	{
+		// The device moved from one platform user to another — update our cached owner.
+		(*Existing)->UserId = NewUserId;
+	}
+	else
+	{
+		// First time seeing this device via a pairing change — add it to the cache.
+		TSharedPtr<FKnownDevice> Entry = MakeShared<FKnownDevice>();
+		Entry->UserId = NewUserId;
+		Entry->LastState = IPlatformInputDeviceMapper::Get().GetInputDeviceConnectionState(DeviceId);
+		if (UInputDeviceSubsystem* DeviceSubsystem = UInputDeviceSubsystem::Get())
+		{
+			Entry->HardwareId = DeviceSubsystem->GetInputDeviceHardwareIdentifier(DeviceId);
+		}
+		KnownDevices.Add(DeviceId, MoveTemp(Entry));
+	}
+
+	RebuildDeviceList();
+}
+
 FText SInputDevicesTab::GetConnectionStateText(EInputDeviceConnectionState State) const
 {
 	switch (State)
diff --git a/Engine/Plugins/Experimental/PlayerInputDebugger/Source/PlayerInputDebugger/Private/SInputDevicesTab.h b/Engine/Plugins/Experimental/PlayerInputDebugger/Source/PlayerInputDebugger/Private/SInputDevicesTab.h
--- a/Engine/Plugins/Experimental/PlayerInputDebugger/Source/PlayerInputDebugger/Private/SInputDevicesTab.h
+++ b/Engine/Plugins/Experimental/PlayerInputDebugger/Source/PlayerInputDebugger/Private/SInputDevicesTab.h
@@ -36,6 +36,7 @@ private:
 	void RebuildDeviceList();
 	void OnDeviceConnectionChanged(EInputDeviceConnectionState NewState, FPlatformUserId UserId, FInputDeviceId DeviceId);
 	void OnHardwareDeviceChanged(FPlatformUserId UserId, FInputDeviceId DeviceId);
+	void OnDevicePairingChanged(FInputDeviceId DeviceId, FPlatformUserId NewUserId, FPlatformUserId OldUserId);
 	FText GetConnectionStateText(EInputDeviceConnectionState State) const;
 	FSlateColor GetConnectionStateColor(EInputDeviceConnectionState State) const;
 	FText GetDeviceTypeText(const FHardwareDeviceIdentifier& HardwareId) const;
@@ -43,6 +44,7 @@ private:
 	TSharedPtr<SVerticalBox> DeviceListContainer;
 	FDelegateHandle ConnectionChangedHandle;
 	FDelegateHandle HardwareDeviceChangedHandle;
+	FDelegateHandle PairingChangeHandle;
 	TMap<FInputDeviceId, TSharedPtr<FKnownDevice>> KnownDevices;
 	TWeakObjectPtr<APlayerController> WeakPC;
 };
