[Player input debugger] Fix a wierd selection issue which would constantly refresh and remove your selection

#rnx
#Synced-CL 52059832

diff --git a/Engine/Plugins/Experimental/PlayerInputDebugger/Source/PlayerInputDebugger/Private/SInputDevicesTab.cpp b/Engine/Plugins/Experimental/PlayerInputDebugger/Source/PlayerInputDebugger/Private/SInputDevicesTab.cpp
--- a/Engine/Plugins/Experimental/PlayerInputDebugger/Source/PlayerInputDebugger/Private/SInputDevicesTab.cpp
+++ b/Engine/Plugins/Experimental/PlayerInputDebugger/Source/PlayerInputDebugger/Private/SInputDevicesTab.cpp
@@ -664,11 +664,18 @@ void SInputDevicesTab::OnDeviceConnectionChanged(EInputDeviceConnectionState New
 
 void SInputDevicesTab::OnHardwareDeviceChanged(FPlatformUserId UserId, FInputDeviceId DeviceId)
 {
+	// This delegate fires every tick while input is being processed, so bail out early
+	// if the hardware identifier hasn't actually changed to avoid thrashing the list view.
 	if (TSharedPtr<FKnownDevice>* Existing = KnownDevices.Find(DeviceId))
 	{
 		if (UInputDeviceSubsystem* DeviceSubsystem = UInputDeviceSubsystem::Get())
 		{
-			(*Existing)->HardwareId = DeviceSubsystem->GetInputDeviceHardwareIdentifier(DeviceId);
+			const FHardwareDeviceIdentifier NewId = DeviceSubsystem->GetInputDeviceHardwareIdentifier(DeviceId);
+			if (NewId == (*Existing)->HardwareId)
+			{
+				return;
+			}
+			(*Existing)->HardwareId = NewId;
 		}
 		RebuildDeviceList();
 		if (SelectedDeviceId.IsSet() && SelectedDeviceId.GetValue() == DeviceId)
