[Player input debugger] Fix localization warnings

#rnx
#Synced-CL 52309720

diff --git a/Engine/Plugins/Experimental/PlayerInputDebugger/Source/PlayerInputDebugger/Private/SInputDevicesTab.cpp b/Engine/Plugins/Experimental/PlayerInputDebugger/Source/PlayerInputDebugger/Private/SInputDevicesTab.cpp
--- a/Engine/Plugins/Experimental/PlayerInputDebugger/Source/PlayerInputDebugger/Private/SInputDevicesTab.cpp
+++ b/Engine/Plugins/Experimental/PlayerInputDebugger/Source/PlayerInputDebugger/Private/SInputDevicesTab.cpp
@@ -780,11 +780,11 @@ void SInputDevicesTab::RebuildDetails()
 		GetDeviceTypeText(Device.HardwareId), FLinearColor(0.7f, 0.7f, 0.7f));
 	AddTextRow(LOCTEXT("DetailInputClass",   "Input Class:"),
 		Device.HardwareId.InputClassName.IsNone()
-			? LOCTEXT("UnknownInputClass", "(unknown)") : FText::FromName(Device.HardwareId.InputClassName),
+			? LOCTEXT("DetailUnknownInputClass", "(unknown)") : FText::FromName(Device.HardwareId.InputClassName),
 		FLinearColor(0.7f, 0.7f, 0.7f));
 	AddTextRow(LOCTEXT("DetailHardwareId",   "Hardware Identifier:"),
 		Device.HardwareId.HardwareDeviceIdentifier.IsNone()
-			? LOCTEXT("UnknownHardware", "(unknown)") : FText::FromName(Device.HardwareId.HardwareDeviceIdentifier),
+			? LOCTEXT("DetailUnknownHardware", "(unknown)") : FText::FromName(Device.HardwareId.HardwareDeviceIdentifier),
 		FLinearColor(0.7f, 0.7f, 0.7f));
 	AddTextRow(LOCTEXT("DetailFeatures",     "Supported Features:"),
 		BuildSupportedFeaturesText(Device.HardwareId.SupportedFeaturesMask), FLinearColor(0.7f, 0.7f, 0.9f));
